- Reworked on effector to goal to support dynamic size of effectors more reliably - GUID to goal mapping system and goal is mapped to the target data
- WIP: FBIK solver is not functioning yet. More focused on the interface
- WIP: refactor how IKRigController is created

diff --git a/Engine/Plugins/Animation/IKRig/Source/IKRig/Private/IKRigDefinition.cpp b/Engine/Plugins/Animation/IKRig/Source/IKRig/Private/IKRigDefinition.cpp
--- a/Engine/Plugins/Animation/IKRig/Source/IKRig/Private/IKRigDefinition.cpp
+++ b/Engine/Plugins/Animation/IKRig/Source/IKRig/Private/IKRigDefinition.cpp
@@ -380,4 +380,13 @@ void UIKRigDefinition::PostLoad()
 	Super::PostLoad();
 
 	Hierarchy.RebuildCacheData();
+}
+
+void UIKRigDefinition::BeginDestroy()
+{
+#if WITH_EDITOR
+	IKRigDefinitionBeginDestroy.Broadcast(this);
+#endif // #if WITH_EDITOR
+
+	Super::BeginDestroy();
 }
\ No newline at end of file
diff --git a/Engine/Plugins/Animation/IKRig/Source/IKRig/Private/IKRigSolver.cpp b/Engine/Plugins/Animation/IKRig/Source/IKRig/Private/IKRigSolver.cpp
--- a/Engine/Plugins/Animation/IKRig/Source/IKRig/Private/IKRigSolver.cpp
+++ b/Engine/Plugins/Animation/IKRig/Source/IKRig/Private/IKRigSolver.cpp
@@ -37,11 +37,11 @@ void UIKRigSolver::Solve(FIKRigTransformModifier& InOutGlobalTransform)
 	}
 }
 
-bool UIKRigSolver::GetTaskTarget(const FName& TaskName, FIKRigTarget& OutTarget) const
+bool UIKRigSolver::GetEffectorTarget(const FIKRigEffector& InEffector, FIKRigTarget& OutTarget) const
 {
 	if (SolverDefinition && GoalGetter.IsBound())
 	{	
-		const FName* GoalName = SolverDefinition->GetTaskToGoal().Find(TaskName);
+		const FName* GoalName = SolverDefinition->GetEffectorToGoal().Find(InEffector);
 		if (GoalName)
 		{
 			return GoalGetter.Execute(*GoalName, OutTarget);
diff --git a/Engine/Plugins/Animation/IKRig/Source/IKRig/Private/IKRigSolverDefinition.cpp b/Engine/Plugins/Animation/IKRig/Source/IKRig/Private/IKRigSolverDefinition.cpp
--- a/Engine/Plugins/Animation/IKRig/Source/IKRig/Private/IKRigSolverDefinition.cpp
+++ b/Engine/Plugins/Animation/IKRig/Source/IKRig/Private/IKRigSolverDefinition.cpp
@@ -13,7 +13,7 @@ void UIKRigSolverDefinition::CollectGoals(TArray<FName>& OutGoals)
 	TArray<FName> LocalGoals;
 	// this just accumulate on the current array
 	// if you want to clear, clear before coming here
-	TaskToGoal.GenerateValueArray(LocalGoals);
+	EffectorToGoal.GenerateValueArray(LocalGoals);
 
 	OutGoals += LocalGoals;
 }
@@ -21,7 +21,7 @@ void UIKRigSolverDefinition::CollectGoals(TArray<FName>& OutGoals)
 #if WITH_EDITOR
 void UIKRigSolverDefinition::RenameGoal(const FName& OldName, const FName& NewName)
 {
-	for (auto Iter = TaskToGoal.CreateIterator(); Iter; ++Iter)
+	for (auto Iter = EffectorToGoal.CreateIterator(); Iter; ++Iter)
 	{
 		if (Iter.Value() == OldName)
 		{
@@ -58,6 +58,20 @@ void UIKRigSolverDefinition::OnGoalHasBeenUpdated()
 	GoalNeedsUpdateDelegate.Broadcast();
 }
 
+void UIKRigSolverDefinition::EnsureToAddEffector(const FIKRigEffector& InEffector, const FString& InPrefix)
+{
+	FName* GoalName = EffectorToGoal.Find(InEffector);
+	if (!GoalName)
+	{
+		EffectorToGoal.Add(InEffector) = CreateUniqueGoalName(*InPrefix);
+	}
+}
+
+void UIKRigSolverDefinition::EnsureToRemoveEffector(const FIKRigEffector& InEffector)
+{
+	EffectorToGoal.Remove(InEffector);
+}
+
 #endif // WITH_EDITOR
 
 void UIKRigSolverDefinition::PostLoad()
@@ -65,6 +79,14 @@ void UIKRigSolverDefinition::PostLoad()
 	Super::PostLoad();
 
 #if WITH_EDITOR
-	UpdateTaskList();
+	UpdateEffectors();
 #endif // WITH_EDITOR
 }
+
+void UIKRigSolverDefinition::Serialize(FArchive& Ar)
+{
+	Super::Serialize(Ar);
+
+	// serialize these manually because that's the custom types
+	Ar << EffectorToGoal;
+}
diff --git a/Engine/Plugins/Animation/IKRig/Source/IKRig/Private/Solvers/TransformSolver.cpp b/Engine/Plugins/Animation/IKRig/Source/IKRig/Private/Solvers/TransformSolver.cpp
--- a/Engine/Plugins/Animation/IKRig/Source/IKRig/Private/Solvers/TransformSolver.cpp
+++ b/Engine/Plugins/Animation/IKRig/Source/IKRig/Private/Solvers/TransformSolver.cpp
@@ -32,7 +32,7 @@ void UTransformSolver::SolveInternal(FIKRigTransformModifier& InOutGlobalTransfo
 	if (UTransformSolverDefinition* SolverDef = Cast<UTransformSolverDefinition>(SolverDefinition))
 	{
 		FIKRigTarget Target;
-		if (GetTaskTarget(UTransformSolverDefinition::TransformTarget, Target))
+		if (GetEffectorTarget(SolverDef->TransformTarget, Target))
 		{
 			int32 Index = InOutGlobalTransform.Hierarchy->GetIndex(Target.Bone);
 			if (Index != INDEX_NONE)
diff --git a/Engine/Plugins/Animation/IKRig/Source/IKRig/Private/Solvers/TransformSolverDefinition.cpp b/Engine/Plugins/Animation/IKRig/Source/IKRig/Private/Solvers/TransformSolverDefinition.cpp
--- a/Engine/Plugins/Animation/IKRig/Source/IKRig/Private/Solvers/TransformSolverDefinition.cpp
+++ b/Engine/Plugins/Animation/IKRig/Source/IKRig/Private/Solvers/TransformSolverDefinition.cpp
@@ -8,30 +8,28 @@
 #include "Solvers/TransformSolverDefinition.h"
 #include "TransformSolver.h"
 
-const FName UTransformSolverDefinition::TransformTarget = FName(TEXT("TransformTarget"));
-
 UTransformSolverDefinition::UTransformSolverDefinition()
+	: TransformTargetName(TEXT("TransformTarget"))
 {
 	DisplayName = TEXT("Transform Solver");
 	ExecutionClass = UTransformSolver::StaticClass();
 }
 
 #if WITH_EDITOR
-void UTransformSolverDefinition::UpdateTaskList()
+void UTransformSolverDefinition::UpdateEffectors()
 {
 	// update tasks - you want to keep old name though if changed
-	FName* TaskFound= TaskToGoal.Find(TransformTarget);
 	bool bActiveTask = bEnablePosition || bEnableRotation;
 
 	// not found, we add new
-	if (bActiveTask && !TaskFound)
+	if (bActiveTask)
 	{
-		TaskToGoal.Add(TransformTarget) = CreateUniqueGoalName(*TransformTarget.ToString());
+		EnsureToAddEffector(TransformTarget, *TransformTargetName);
 	}
 	// if the opposite, we have to remove
-	else if (!bActiveTask && TaskFound)
+	else if (!bActiveTask)
 	{
-		TaskToGoal.FindAndRemoveChecked(TransformTarget);
+		EnsureToRemoveEffector(TransformTarget);
 	}
 
 	// trigger a delegate for goal has been updated?
@@ -43,8 +41,10 @@ void UTransformSolverDefinition::PostEditChangeChainProperty(struct FPropertyCha
 	if (PropertyChangedEvent.Property->GetFName() == GET_MEMBER_NAME_CHECKED(UTransformSolverDefinition, bEnablePosition) ||
 		PropertyChangedEvent.Property->GetFName() == GET_MEMBER_NAME_CHECKED(UTransformSolverDefinition, bEnableRotation))
 	{
-		UpdateTaskList();
+		UpdateEffectors();
 	}
+
+	Super::PostEditChangeChainProperty(PropertyChangedEvent);
 }
 
 #endif// WITH_EDITOR
\ No newline at end of file
diff --git a/Engine/Plugins/Animation/IKRig/Source/IKRig/Public/IKRigDataTypes.h b/Engine/Plugins/Animation/IKRig/Source/IKRig/Public/IKRigDataTypes.h
--- a/Engine/Plugins/Animation/IKRig/Source/IKRig/Public/IKRigDataTypes.h
+++ b/Engine/Plugins/Animation/IKRig/Source/IKRig/Public/IKRigDataTypes.h
@@ -11,6 +11,7 @@
 #include "CoreMinimal.h"
 #include "UObject/ObjectMacros.h"
 #include "UObject/Object.h"
+#include "Misc/Guid.h"
 #include "IKRigDataTypes.generated.h"
 
 struct FIKRigHierarchy;
@@ -92,6 +93,40 @@ struct IKRIG_API FIKRigGoal
 	}
 };
 
+USTRUCT()
+struct IKRIG_API FIKRigEffector
+{
+	GENERATED_BODY()
+
+	UPROPERTY()
+	FGuid Guid;
+
+	FIKRigEffector()
+	{
+		Guid = FGuid::NewGuid();
+	}
+
+	friend FArchive& operator<<(FArchive& Ar, FIKRigEffector& Effector)
+	{
+		return Ar << Effector.Guid;
+	}
+
+	FORCEINLINE bool operator==(const FIKRigEffector& Rhs) const
+	{
+		return Guid == Rhs.Guid;
+	}
+};
+
+template <typename ValueType>
+struct TIKRigEffectorMapKeyFuncs : public TDefaultMapKeyFuncs<const FIKRigEffector, ValueType, false>
+{
+	static FORCEINLINE FIKRigEffector					GetSetKey(TPair<FIKRigEffector, ValueType> const& Element) { return Element.Key; }
+	static FORCEINLINE uint32							GetKeyHash(FIKRigEffector const& Key) { return GetTypeHash(Key.Guid); }
+	static FORCEINLINE bool								Matches(FIKRigEffector const& A, FIKRigEffector const& B) { return (A.Guid == B.Guid); }
+};
+
+template <typename ValueType>
+using TIKRigEffectorMap = TMap<FIKRigEffector, ValueType, FDefaultSetAllocator, TIKRigEffectorMapKeyFuncs<ValueType>>;
 
 // If you want to modify this, use FIKRigTransformModifier
 USTRUCT()
diff --git a/Engine/Plugins/Animation/IKRig/Source/IKRig/Public/IKRigDefinition.h b/Engine/Plugins/Animation/IKRig/Source/IKRig/Public/IKRigDefinition.h
--- a/Engine/Plugins/Animation/IKRig/Source/IKRig/Public/IKRigDefinition.h
+++ b/Engine/Plugins/Animation/IKRig/Source/IKRig/Public/IKRigDefinition.h
@@ -103,9 +103,10 @@ private:
 	UPROPERTY(EditAnywhere, Category = "Goals")
 	TMap<FName, FIKRigGoal> IKGoals;
 
-	// UObject functions
+	// BEGIN UObject functions 
 	virtual void PostLoad() override;
-
+	virtual void BeginDestroy() override;
+	// END UObject functions
 #if WITH_EDITOR
 	// https ://drive.google.com/file/d/1EO0Ijojx-0jommdHuZv1JYxfCdoFLd44/view
 
@@ -139,13 +140,13 @@ private:
 	DECLARE_MULTICAST_DELEGATE_TwoParams(FBoneRenamed, FName /*OldName*/, FName /*NewName*/);
 	DECLARE_MULTICAST_DELEGATE_TwoParams(FBoneReparented, FName /*Name*/, FName /*ParentName*/);
 
-	DECLARE_MULTICAST_DELEGATE(FIKRigDefinitionModified);
+	DECLARE_MULTICAST_DELEGATE_OneParam(FIKRigDefinitionBeginDestroy, UIKRigDefinition* Definition);
 
 	FBoneAdded BoneAddedDelegate;
 	FBoneRemoved BoneRemovedDelegate;
 	FBoneRenamed BoneRenamedDelegate;
 	FBoneReparented BoneReparentedDelegate;
-	FIKRigDefinitionModified IKRigDefinitionModified;
+	FIKRigDefinitionBeginDestroy IKRigDefinitionBeginDestroy;
 
 #endif // WITH_EDITOR
 
diff --git a/Engine/Plugins/Animation/IKRig/Source/IKRig/Public/IKRigSolver.h b/Engine/Plugins/Animation/IKRig/Source/IKRig/Public/IKRigSolver.h
--- a/Engine/Plugins/Animation/IKRig/Source/IKRig/Public/IKRigSolver.h
+++ b/Engine/Plugins/Animation/IKRig/Source/IKRig/Public/IKRigSolver.h
@@ -12,6 +12,7 @@
 #include "UObject/Object.h"
 #include "IKRigSolver.generated.h"
 
+struct FIKRigEffector;
 struct FIKRigTarget;
 struct FIKRigTransform;
 struct FIKRigTransformModifier;
@@ -51,7 +52,7 @@ protected:
 	virtual void SolveInternal(FIKRigTransformModifier& InOutGlobalTransform) {};
 	virtual bool IsSolverActive() const;
 
-	bool GetTaskTarget(const FName& TaskName, FIKRigTarget& OutTarget) const;
+	bool GetEffectorTarget(const FIKRigEffector& InEffector, FIKRigTarget& OutTarget) const;
 	const FIKRigTransform& GetReferencePose() const;
 private:
 	// delegate
diff --git a/Engine/Plugins/Animation/IKRig/Source/IKRig/Public/IKRigSolverDefinition.h b/Engine/Plugins/Animation/IKRig/Source/IKRig/Public/IKRigSolverDefinition.h
--- a/Engine/Plugins/Animation/IKRig/Source/IKRig/Public/IKRigSolverDefinition.h
+++ b/Engine/Plugins/Animation/IKRig/Source/IKRig/Public/IKRigSolverDefinition.h
@@ -12,7 +12,7 @@
 #include "Templates/SubclassOf.h"
 #include "UObject/ObjectMacros.h"
 #include "UObject/Object.h"
-
+#include "IKRigDataTypes.h"
 #include "IKRigSolverDefinition.generated.h"
 
 class UIKRigSolver;
@@ -27,10 +27,6 @@ protected:
 	UPROPERTY(EditAnywhere, Category = "Definition")
 	FString DisplayName;
 
-	// effector name to goals name map
-	UPROPERTY(VisibleAnywhere, Category = "Definition")
-	TMap<FName, FName> TaskToGoal;
-
 	// this is name of constraints it's used
 	// use Delegate from IKRigDefinition to query and solve
 	UPROPERTY(EditAnywhere, Category = "Definition")
@@ -39,10 +35,13 @@ protected:
 	UPROPERTY(VisibleAnywhere, Category = "Definition")
 	TSubclassOf<UIKRigSolver> ExecutionClass;
 
+	// effector name to goals name map
+	TIKRigEffectorMap<FName> EffectorToGoal;
+
 public:
-	const TMap<FName, FName>& GetTaskToGoal() const
+	const TIKRigEffectorMap<FName>& GetEffectorToGoal() const
 	{
-		return TaskToGoal;
+		return EffectorToGoal;
 	}
 
 	TSubclassOf<UIKRigSolver> GetExecutionClass() 
@@ -68,19 +67,22 @@ private:
 	void RenameGoal(const FName& OldName, const FName& NewName);
 	// get unique name delegate by IKRigDefinition
 	void EnsureUniqueGoalName(FName& InOutGoalName) const;
-
+	
 	DECLARE_MULTICAST_DELEGATE(FGoalNeedsUpdate);
 	FGoalNeedsUpdate GoalNeedsUpdateDelegate;
 
 protected:
 	FName CreateUniqueGoalName(const TCHAR* Suffix) const;
 	void OnGoalHasBeenUpdated();
+	void EnsureToAddEffector(const FIKRigEffector& InEffector, const FString& InPrefix);
+	void EnsureToRemoveEffector(const FIKRigEffector& InEffector);
 
-	virtual void UpdateTaskList() {};
+	virtual void UpdateEffectors() {};
 #endif // WITH_EDITOR
 
 	/// BEGIN UObject
 	virtual void PostLoad() override;
+	virtual void Serialize(FArchive& Ar) override;
 	/// END UObject
 	friend class UIKRigController;
 };
diff --git a/Engine/Plugins/Animation/IKRig/Source/IKRig/Public/Solvers/TransformSolverDefinition.h b/Engine/Plugins/Animation/IKRig/Source/IKRig/Public/Solvers/TransformSolverDefinition.h
--- a/Engine/Plugins/Animation/IKRig/Source/IKRig/Public/Solvers/TransformSolverDefinition.h
+++ b/Engine/Plugins/Animation/IKRig/Source/IKRig/Public/Solvers/TransformSolverDefinition.h
@@ -25,15 +25,17 @@ public:
 	UPROPERTY(EditAnywhere, Category = "Solver")
 	bool	bEnableRotation = true;
 
+	UPROPERTY(EditAnywhere, Category = "Solver")
+	FIKRigEffector TransformTarget;
+
 private:
 #if WITH_EDITOR
-	virtual void UpdateTaskList() override;
+	virtual void UpdateEffectors() override;
 
 	// UObject interface
 	virtual void PostEditChangeChainProperty(struct FPropertyChangedChainEvent& PropertyChangedEvent) override;
 	// UObject interface
 #endif // WITH_EDITOR
 
-public:
-	static const FName TransformTarget;
+	const FString TransformTargetName;
 };
\ No newline at end of file
diff --git a/Engine/Plugins/Animation/IKRig/Source/IKRigEditor/Private/IKRigController.cpp b/Engine/Plugins/Animation/IKRig/Source/IKRigEditor/Private/IKRigController.cpp
--- a/Engine/Plugins/Animation/IKRig/Source/IKRigEditor/Private/IKRigController.cpp
+++ b/Engine/Plugins/Animation/IKRig/Source/IKRigEditor/Private/IKRigController.cpp
@@ -10,6 +10,53 @@
 #include "AnimationRuntime.h"
 #include "IKRigSolverDefinition.h"
 
+TMap<UIKRigDefinition*, UIKRigController*> UIKRigController::DefinitionToControllerMap;
+
+// currently it's not clear the lifecycle of this map
+// usually during editor, they will be present, 
+UIKRigController* UIKRigController::GetControllerByRigDefinition(UIKRigDefinition* InIKRigDefinition)
+{
+	if (InIKRigDefinition)
+	{
+		UIKRigController** Controller = DefinitionToControllerMap.Find(InIKRigDefinition);
+		if (Controller)
+		{
+			return *Controller;
+		}
+		else
+		{
+			UIKRigController* NewController = NewObject<UIKRigController>();
+			DefinitionToControllerMap.Add(InIKRigDefinition) = NewController;
+			NewController->SetIKRigDefinition(InIKRigDefinition);
+			InIKRigDefinition->IKRigDefinitionBeginDestroy.AddStatic(&UIKRigController::RemoveControllerByRigDefinition);
+			return NewController;
+		}
+	}
+
+	return nullptr;
+}
+
+// this shoudl be called by IKRigDefinition::BeginDestroy;
+void UIKRigController::RemoveControllerByRigDefinition(UIKRigDefinition* InIKRigDefinition)
+{
+	DefinitionToControllerMap.Remove(InIKRigDefinition);
+}
+
+void UIKRigController::AddReferencedObjects(UObject* InThis, FReferenceCollector& Collector)
+{
+	// since static member, we just add only for default object
+	if (InThis && InThis->IsTemplate())
+	{
+		for (auto Iter = DefinitionToControllerMap.CreateIterator(); Iter; ++Iter)
+		{
+			// we add controllers to the list. so that it doens't GCed
+			// when do we delete from the list?
+			UIKRigController* Obj = Iter.Value();
+			Collector.AddReferencedObject(Obj);
+		}
+	}
+}
+
 void UIKRigController::BeginDestroy() 
 {
 	if (IKRigDefinition)
@@ -55,7 +102,7 @@ void UIKRigController::InitializeIKRigSolverDefinition(UIKRigSolverDefinition* S
 		FDelegateHandle DelegateHandle = SolverDef->GoalNeedsUpdateDelegate.AddUObject(this, &UIKRigController::UpdateGoal);
 		SolverDelegateHandles.Add(SolverDef) = DelegateHandle;
 
-		SolverDef->UpdateTaskList();
+		SolverDef->UpdateEffectors();
 	}
 }
 
diff --git a/Engine/Plugins/Animation/IKRig/Source/IKRigEditor/Private/IKRigDefinitionDetails.cpp b/Engine/Plugins/Animation/IKRig/Source/IKRigEditor/Private/IKRigDefinitionDetails.cpp
--- a/Engine/Plugins/Animation/IKRig/Source/IKRigEditor/Private/IKRigDefinitionDetails.cpp
+++ b/Engine/Plugins/Animation/IKRig/Source/IKRigEditor/Private/IKRigDefinitionDetails.cpp
@@ -64,13 +64,8 @@ void FIKRigDefinitionDetails::CustomizeDetails(IDetailLayoutBuilder& DetailBuild
 	}
 
 	// create controller
-	if (!IKRigController.IsValid())
-	{
-		IKRigController = TStrongObjectPtr<UIKRigController>(NewObject<UIKRigController>());
-	}
+	IKRigController = UIKRigController::GetControllerByRigDefinition(IKRigDefinition.Get());
 	
-	IKRigController->SetIKRigDefinition(IKRigDefinition.Get());
-
 	ObjectChangedDelegate = FCoreUObjectDelegates::OnObjectPropertyChanged.AddRaw(this, &FIKRigDefinitionDetails::OnObjectPostEditChange);
 	/////////////////////////////////////////////////////////////////////////////////
 	// skeleton set up
@@ -319,7 +314,7 @@ TSharedRef<ITableRow> FIKRigDefinitionDetails::OnGenerateWidgetForGoals(FGoalNam
 void FIKRigDefinitionDetails::HandleGoalNameChanged(const FText& NewName, ETextCommit::Type CommitType, FGoalNameListItemPtr InItem)
 {
 	//if (CommitType == ETextCommit::OnEnter)
-	if (IKRigController.IsValid())
+	if (IKRigController)
 	{
 		if (!NewName.IsEmptyOrWhitespace())
 		{
diff --git a/Engine/Plugins/Animation/IKRig/Source/IKRigEditor/Private/IKRigDefinitionDetails.h b/Engine/Plugins/Animation/IKRig/Source/IKRigEditor/Private/IKRigDefinitionDetails.h
--- a/Engine/Plugins/Animation/IKRig/Source/IKRigEditor/Private/IKRigDefinitionDetails.h
+++ b/Engine/Plugins/Animation/IKRig/Source/IKRigEditor/Private/IKRigDefinitionDetails.h
@@ -3,7 +3,7 @@
 #pragma once
 
 #include "CoreMinimal.h"
-#include "UObject/StrongObjectPtr.h"
+#include "UObject/WeakObjectPtrTemplates.h"
 #include "IDetailCustomization.h"
 
 class IDetailLayoutBuilder;
@@ -36,7 +36,7 @@ public:
 	virtual void CustomizeDetails(const TSharedPtr<IDetailLayoutBuilder>& DetailBuilder) override;
 private:
 	TWeakObjectPtr<UIKRigDefinition>		IKRigDefinition;
-	TStrongObjectPtr<UIKRigController>		IKRigController;
+	UIKRigController*						IKRigController;
 
 	// source asset
 	FString GetCurrentSourceAsset() const;
diff --git a/Engine/Plugins/Animation/IKRig/Source/IKRigEditor/Private/IKRigEditor.cpp b/Engine/Plugins/Animation/IKRig/Source/IKRigEditor/Private/IKRigEditor.cpp
--- a/Engine/Plugins/Animation/IKRig/Source/IKRigEditor/Private/IKRigEditor.cpp
+++ b/Engine/Plugins/Animation/IKRig/Source/IKRigEditor/Private/IKRigEditor.cpp
@@ -16,6 +16,7 @@
 #include "AnimGraphNode_IKRig.h"
 #include "AssetTypeActions_IKRigDefinition.h"
 #include "IKRigDefinition.h"
+#include "AnimMode/IKRigEditMode.h"
 
 IMPLEMENT_MODULE(FIKRigEditor, IKRigEditor)
 
@@ -31,12 +32,16 @@ void FIKRigEditor::StartupModule()
 	FPropertyEditorModule& PropertyModule = FModuleManager::GetModuleChecked<FPropertyEditorModule>("PropertyEditor");
 	PropertyModule.RegisterCustomClassLayout(UIKRigDefinition::StaticClass()->GetFName(), FOnGetDetailCustomizationInstance::CreateStatic(&FIKRigDefinitionDetails::MakeInstance));
 
-//	FEditorModeRegistry::Get().RegisterMode<FIKRigEditMode>(UAnimGraphNode_IKRig::AnimModeName, LOCTEXT("IKRigEditMode", "IKRig"), FSlateIcon(), false);
+//	PropertyModule.RegisterCustomPropertyTypeLayout("IKRigEffector", FOnGetPropertyTypeCustomizationInstance::CreateStatic(&FIKRigEffectorCustomization::MakeInstance));
+
+	FEditorModeRegistry::Get().RegisterMode<FIKRigEditMode>(UAnimGraphNode_IKRig::AnimModeName, LOCTEXT("IKRigEditMode", "IKRig"), FSlateIcon(), false);
 }
 
 void FIKRigEditor::ShutdownModule()
 {
-	//	FEditorModeRegistry::Get().UnregisterMode(UAnimGraphNode_IKRig::AnimModeName);
+	FEditorModeRegistry::Get().UnregisterMode(UAnimGraphNode_IKRig::AnimModeName);
+
+//	PropertyModule.UnregisterCustomPropertyTypeLayout("IKRigEffector");
 
 	FPropertyEditorModule& PropertyModule = FModuleManager::GetModuleChecked<FPropertyEditorModule>("PropertyEditor");
 	PropertyModule.UnregisterCustomClassLayout(UIKRigDefinition::StaticClass()->GetFName());
diff --git a/Engine/Plugins/Animation/IKRig/Source/IKRigEditor/Public/IKRigController.h b/Engine/Plugins/Animation/IKRig/Source/IKRigEditor/Public/IKRigController.h
--- a/Engine/Plugins/Animation/IKRig/Source/IKRigEditor/Public/IKRigController.h
+++ b/Engine/Plugins/Animation/IKRig/Source/IKRigEditor/Public/IKRigController.h
@@ -32,9 +32,6 @@ class IKRIGEDITOR_API UIKRigController : public UObject
 	GENERATED_BODY()
 
 public:
-	// IKRigDefinition set up
-	void SetIKRigDefinition(UIKRigDefinition* InIKRigDefinition);
-
 	// hierarchy operators
 	void SetSkeleton(const FReferenceSkeleton& InSkeleton);
 	// const hierarchy getter
@@ -87,7 +84,12 @@ public:
 
 	// BEGIN UObject
 	virtual void BeginDestroy() override;
+	static void AddReferencedObjects(UObject* InThis, FReferenceCollector& Collector);
 	// END UObject
+
+	// Controller getter
+	// Use this interface to get controller 
+	static UIKRigController* GetControllerByRigDefinition(UIKRigDefinition* InIKRigDefinition);
 private:
 	UIKRigDefinition* IKRigDefinition = nullptr;
 
@@ -98,5 +100,17 @@ private:
 	void UninitializeIKRigSolverDefinition(UIKRigSolverDefinition* SolverDef);
 
 	TMap<UIKRigSolverDefinition*, FDelegateHandle> SolverDelegateHandles;
+
+	static TMap<UIKRigDefinition*, UIKRigController*> DefinitionToControllerMap;
+
+	// IKRigDefinition set up
+	void SetIKRigDefinition(UIKRigDefinition* InIKRigDefinition);
+
+	// called by IKRigDefinition when it's begin destroy
+	static void RemoveControllerByRigDefinition(UIKRigDefinition* InIKRigDefinition);
+
+	// friend class 
+	friend class UIKRigDefinition;
 };
 
+   
\ No newline at end of file
diff --git a/Engine/Plugins/Experimental/FullBodyIK/FullBodyIK.uplugin b/Engine/Plugins/Experimental/FullBodyIK/FullBodyIK.uplugin
--- a/Engine/Plugins/Experimental/FullBodyIK/FullBodyIK.uplugin
+++ b/Engine/Plugins/Experimental/FullBodyIK/FullBodyIK.uplugin
@@ -18,6 +18,10 @@
 		{
 			"Name": "ControlRig",
 			"Enabled": true
+		},
+		{
+			"Name": "IKRig",
+			"Enabled": true
 		}
 	],
 	"Modules": [
diff --git a/Engine/Plugins/Experimental/FullBodyIK/Source/FullBodyIK/FullBodyIK.Build.cs b/Engine/Plugins/Experimental/FullBodyIK/Source/FullBodyIK/FullBodyIK.Build.cs
--- a/Engine/Plugins/Experimental/FullBodyIK/Source/FullBodyIK/FullBodyIK.Build.cs
+++ b/Engine/Plugins/Experimental/FullBodyIK/Source/FullBodyIK/FullBodyIK.Build.cs
@@ -15,7 +15,8 @@ public class FullBodyIK : ModuleRules
 				// ... add other public dependencies that you statically link with here ...
 				"Eigen",
 				"RigVM",
-				"ControlRig"
+				"ControlRig",
+				"IKRig"
             }
 			);
 			
diff --git a/Engine/Plugins/Experimental/FullBodyIK/Source/FullBodyIK/Private/FBIKShared.cpp b/Engine/Plugins/Experimental/FullBodyIK/Source/FullBodyIK/Private/FBIKShared.cpp
new file mode 100644
--- /dev/null
+++ b/Engine/Plugins/Experimental/FullBodyIK/Source/FullBodyIK/Private/FBIKShared.cpp
@@ -0,0 +1,50 @@
+// Copyright Epic Games, Inc. All Rights Reserved.
+
+#include "FBIKShared.h"
+#include "Rigs/RigHierarchyContainer.h"
+#include "FBIKUtil.h"
+
+void FJacobianSolver_FullbodyIK::InitializeSolver(TArray<FFBIKLinkData>& InOutLinkData, TMap<int32, FFBIKEffectorTarget>& InOutEffectors) const
+{
+	// since we're using constraints, we don't want to create new motion base all the time
+	// constraints will add that info and we'll utilize it
+}
+
+void FJacobianSolver_FullbodyIK::PreSolve(TArray<FFBIKLinkData>& InOutLinkData, const TMap<int32, FFBIKEffectorTarget>& InEffectors) const
+{
+	// for final one, we'll have to add those in the beginning, and then just update the frame for every iteration
+	for (int32 Index = 0; Index < InOutLinkData.Num(); ++Index)
+	{
+		// we update existing data
+		const FQuat& LocalFrame = InOutLinkData[Index].LocalFrame;
+
+		FQuat CurrentFrame;
+
+		if (InOutLinkData[Index].ParentLinkIndex != INDEX_NONE)
+		{
+			CurrentFrame = InOutLinkData[InOutLinkData[Index].ParentLinkIndex].GetTransform().GetRotation() * LocalFrame;
+		}
+		else
+		{
+			CurrentFrame = LocalFrame;
+		}
+
+		// update the base axis
+		if (InOutLinkData[Index].GetNumMotionBases() == 3)
+		{
+			InOutLinkData[Index].GetMotionBase(0).BaseAxis = CurrentFrame.RotateVector(FVector::ForwardVector);
+			InOutLinkData[Index].GetMotionBase(1).BaseAxis = CurrentFrame.RotateVector(FVector::RightVector);
+			InOutLinkData[Index].GetMotionBase(2).BaseAxis = CurrentFrame.RotateVector(FVector::UpVector);
+		}
+		else
+		{
+			ensure(InOutLinkData[Index].GetNumMotionBases() == 0);
+
+			// invalid number of motion bases for this solver
+			InOutLinkData[Index].ResetMotionBases();
+			InOutLinkData[Index].AddMotionBase(FMotionBase(CurrentFrame.RotateVector(FVector::ForwardVector)));
+			InOutLinkData[Index].AddMotionBase(FMotionBase(CurrentFrame.RotateVector(FVector::RightVector)));
+			InOutLinkData[Index].AddMotionBase(FMotionBase(CurrentFrame.RotateVector(FVector::UpVector)));
+		}
+	}
+}
diff --git a/Engine/Plugins/Experimental/FullBodyIK/Source/FullBodyIK/Private/FBIKShared.h b/Engine/Plugins/Experimental/FullBodyIK/Source/FullBodyIK/Private/FBIKShared.h
new file mode 100644
--- /dev/null
+++ b/Engine/Plugins/Experimental/FullBodyIK/Source/FullBodyIK/Private/FBIKShared.h
@@ -0,0 +1,76 @@
+// Copyright Epic Games, Inc. All Rights Reserved.
+
+#pragma once
+
+#include "CoreMinimal.h"
+#include "JacobianSolver.h"
+#include "Rigs/RigHierarchyDefines.h"
+#include "FBIKShared.generated.h"
+
+class FULLBODYIK_API FJacobianSolver_FullbodyIK : public FJacobianSolverBase
+{
+public:
+	FJacobianSolver_FullbodyIK() : FJacobianSolverBase() {}
+
+private:
+	virtual void InitializeSolver(TArray<FFBIKLinkData>& InOutLinkData, TMap<int32, FFBIKEffectorTarget>& InOutEndEffectors) const override;
+	virtual void PreSolve(TArray<FFBIKLinkData>& InOutLinkData, const TMap<int32, FFBIKEffectorTarget>& InEndEffectors) const override;
+};
+
+USTRUCT()
+struct FSolverInput
+{
+	GENERATED_BODY()
+
+	/*
+	 * This value is applied to the target information for effectors, which influence back to 
+	 * Joint's motion that are affected by the end effector
+	 * The reason min/max is used when we apply the depth through the chain that are affected
+
+	 */
+	UPROPERTY(EditAnywhere, Category = FSolverInput)
+	float	LinearMotionStrength = 3.f;
+
+	UPROPERTY(EditAnywhere, Category = FSolverInput)
+	float	MinLinearMotionStrength = 2.f;
+
+	/*
+	 * This value is applied to the target information for effectors, which influence back to 
+	 * Joint's motion that are affected by the end effector
+	 * The reason min/max is used when we apply the depth through the chain that are affected
+	 */
+	UPROPERTY(EditAnywhere, Category = FSolverInput)
+	float	AngularMotionStrength = 3.f;
+
+	UPROPERTY(EditAnywhere, Category = FSolverInput)
+	float	MinAngularMotionStrength = 2.f;
+
+	/* This is a scale value (range from 0-0.7) that is used to stablize the target vector. If less, it's more stable, but it can reduce speed of converge. */
+	UPROPERTY(EditAnywhere, Category = FSolverInput)
+	float	DefaultTargetClamp = 0.2f;
+
+	/**
+	 * The precision to use for the solver
+	 */
+	UPROPERTY(EditAnywhere, Category = FSolverInput)
+	float Precision = 0.1f;
+
+	/**
+	* The precision to use for the fabrik solver
+	*/
+	UPROPERTY(EditAnywhere, Category = FSolverInput)
+	float Damping = 30.f;
+
+	/**
+	 * The maximum number of iterations. Values between 4 and 16 are common.
+	 */
+	UPROPERTY(EditAnywhere, Category = FSolverInput)
+	int32 MaxIterations = 30;
+
+	/**
+	 * Cheaper solution than default Jacobian Pseudo Inverse Damped Least Square
+	 */
+	UPROPERTY(EditAnywhere, Category = FSolverInput)
+	bool bUseJacobianTranspose = false;
+};
+
diff --git a/Engine/Plugins/Experimental/FullBodyIK/Source/FullBodyIK/Private/FullBodyIKSolver.cpp b/Engine/Plugins/Experimental/FullBodyIK/Source/FullBodyIK/Private/FullBodyIKSolver.cpp
new file mode 100644
--- /dev/null
+++ b/Engine/Plugins/Experimental/FullBodyIK/Source/FullBodyIK/Private/FullBodyIKSolver.cpp
@@ -0,0 +1,374 @@
+// Copyright Epic Games, Inc. All Rights Reserved.
+
+/*=============================================================================
+	FullBodyIKSolver.cpp: Solver execution class for Transform
+=============================================================================*/
+
+#include "FullBodyIKSolver.h"
+#include "IKRigDataTypes.h"
+#include "IKRigHierarchy.h"
+#include "FBIKConstraint.h"
+#include "JacobianIK.h"
+#include "FullBodyIKSolverDefinition.h"
+
+void UFullBodyIKSolver::InitInternal()
+{
+	if (UFullBodyIKSolverDefinition* SolverDef = Cast<UFullBodyIKSolverDefinition>(SolverDefinition))
+	{
+// 		LinkData.Reset();
+// 		EffectorTargets.Reset();
+// 		EffectorLinkIndices.Reset();
+// 		LinkDataToHierarchyIndices.Reset();
+// 		HierarchyToLinkDataMap.Reset();
+// 
+// 		// verify the chain
+// 		AddEffectors(Hierarchy, Root, Effectors, LinkData, EffectorTargets, EffectorLinkIndices, LinkDataToHierarchyIndices, HierarchyToLinkDataMap, SolverProperty);
+	}
+}
+
+bool UFullBodyIKSolver::IsSolverActive() const
+{
+	if (Super::IsSolverActive())
+	{
+ 		if (UFullBodyIKSolverDefinition* SolverDef = Cast<UFullBodyIKSolverDefinition>(SolverDefinition))
+ 		{
+			// i'm not checkinf if this is valid, but we may not know for sure because 
+ 			return SolverDef->Effectors.Num() > 1;
+ 		}
+	}
+
+	return false;
+}
+
+void UFullBodyIKSolver::SolveInternal(FIKRigTransformModifier& InOutGlobalTransform)
+{
+#if 0 
+	if (LinkDataToHierarchyIndices.Num() > 0)
+	{
+		// we do this every frame for now
+		if (Constraints.Num() > 0)
+		{
+			DECLARE_SCOPE_HIERARCHICAL_COUNTER(TEXT("Build Constraint"))
+			//Build constraints
+			FBIKConstraintLib::BuildConstraints(Constraints, InternalConstraints, Hierarchy, LinkData, LinkDataToHierarchyIndices, HierarchyToLinkDataMap);
+		}
+
+		// during only editor and update
+		// we expect solver type changes, it will reinit
+		// InternalConstraints can't be changed during runtime
+		if (InternalConstraints.Num() > 0)
+		{
+			WorkData.IKSolver.SetPostProcessDelegateForIteration(FPostProcessDelegateForIteration::CreateStatic(&FBIKConstraintLib::ApplyConstraint, &InternalConstraints));
+		}
+		else
+		{
+			WorkData.IKSolver.ClearPostProcessDelegateForIteration();
+		}
+		// first set mid effector's transform
+
+		// before update we finalize motion scale
+		// this code may go away once we have constraint
+		// update link data and end effectors
+		for (int32 LinkIndex = 0; LinkIndex < LinkData.Num(); ++LinkIndex)
+		{
+			const FRigElementKey& Item = *LinkDataToHierarchyIndices.Find(LinkIndex);
+			LinkData[LinkIndex].SetTransform(Hierarchy->GetGlobalTransform(Item));
+			// @todo: fix this somewhere else - we can add this to prepare step
+			// @todo: we update motion scale here, then?
+			LinkData[LinkIndex].FinalizeForSolver();
+		}
+
+		// update mid effector info
+
+		const float LinearMotionStrength = FMath::Max(SolverProperty.LinearMotionStrength, SolverProperty.MinLinearMotionStrength);
+		const float AngularMotionStrength = FMath::Max(SolverProperty.AngularMotionStrength, SolverProperty.MinAngularMotionStrength);
+		const float LinearRange = LinearMotionStrength - SolverProperty.MinLinearMotionStrength;
+		const float AngularRange = AngularMotionStrength - SolverProperty.MinAngularMotionStrength;
+
+		// update end effector info
+		for (int32 EffectorIndex = 0; EffectorIndex < Effectors.Num(); ++EffectorIndex)
+		{
+			int32 EffectorLinkIndex = EffectorLinkIndices[EffectorIndex];
+			if (EffectorLinkIndex != INDEX_NONE)
+			{
+				FFBIKEffectorTarget* EffectorTarget = EffectorTargets.Find(EffectorLinkIndex);
+				if (EffectorTarget)
+				{
+					const FFBIKEndEffector& CurEffector = Effectors[EffectorIndex];
+					const FVector CurrentLinkLocation = LinkData[EffectorLinkIndex].GetTransform().GetLocation();
+					const FQuat CurrentLinkRotation = LinkData[EffectorLinkIndex].GetTransform().GetRotation();
+					const FVector& EffectorLocation = CurEffector.Position;
+					const FQuat& EffectorRotation = CurEffector.Rotation;
+					EffectorTarget->Position = FMath::Lerp(CurrentLinkLocation, EffectorLocation, CurEffector.PositionAlpha);
+					EffectorTarget->Rotation = FMath::Lerp(CurrentLinkRotation, EffectorRotation, CurEffector.RotationAlpha);
+					EffectorTarget->InitialPositionDistance = (EffectorLocation - CurrentLinkLocation).Size();
+					EffectorTarget->InitialRotationDistance = (FBIKUtil::GetScaledRotationAxis(EffectorRotation) - FBIKUtil::GetScaledRotationAxis(CurrentLinkRotation)).Size();
+
+					const float Pull = FMath::Clamp(CurEffector.Pull, 0.f, 1.f);
+					// we want some impact of Pull, in order for Pull to have some impact, we clamp to some number
+					const float TargetClamp = FMath::Clamp(SolverProperty.DefaultTargetClamp, 0.f, 0.7f);
+					const float Scale = TargetClamp + Pull * (1.f - TargetClamp);
+					// Pull set up
+					EffectorTarget->LinearMotionStrength = LinearRange * Scale + SolverProperty.MinLinearMotionStrength;
+					EffectorTarget->AngularMotionStrength = AngularRange * Scale + SolverProperty.MinAngularMotionStrength;
+					EffectorTarget->ConvergeScale = Scale;
+					EffectorTarget->TargetClampScale = Scale;
+
+					EffectorTarget->bPositionEnabled = true;
+					EffectorTarget->bRotationEnabled = true;
+				}
+			}
+		}
+
+		TArray<FJacobianDebugData>& DebugData = WorkData.DebugData;
+		DebugData.Reset();
+
+		const bool bDebugEnabled = DebugOption.bDrawDebugHierarchy || DebugOption.bDrawDebugEffector || DebugOption.bDrawDebugConstraints;
+
+		// we can't reuse memory until we fix the memory issue on RigVM
+		{
+			DECLARE_SCOPE_HIERARCHICAL_COUNTER(TEXT("Solver"))
+			FJacobianSolver_FullbodyIK& IKSolver = WorkData.IKSolver;
+			IKSolver.SolveJacobianIK(LinkData, EffectorTargets,
+				JacobianIK::FSolverParameter(SolverProperty.Damping, true, false, (SolverProperty.bUseJacobianTranspose) ? EJacobianSolver::JacobianTranspose : EJacobianSolver::JacobianPIDLS),
+				SolverProperty.MaxIterations, SolverProperty.Precision, (bDebugEnabled) ? &DebugData : nullptr);
+
+			if (MotionProperty.bForceEffectorRotationTarget)
+			{
+				// if position is reached, we force rotation target
+				for (int32 EffectorIndex = 0; EffectorIndex < Effectors.Num(); ++EffectorIndex)
+				{
+					int32 EffectorLinkIndex = EffectorLinkIndices[EffectorIndex];
+					if (EffectorLinkIndex != INDEX_NONE)
+					{
+						FFBIKEffectorTarget* EffectorTarget = EffectorTargets.Find(EffectorLinkIndex);
+						if (EffectorTarget && EffectorTarget->bRotationEnabled)
+						{
+							bool bApplyRotation = true;
+
+							if (MotionProperty.bOnlyApplyWhenReachedToTarget)
+							{
+								// only do this when position is reached? This will conflict with converge scale
+								const FVector& BonePosition = LinkData[EffectorLinkIndex].GetTransform().GetLocation();
+								const FVector& TargetPosition = EffectorTarget->Position;
+
+								bApplyRotation = (FVector(BonePosition - TargetPosition).SizeSquared() <= SolverProperty.Precision * SolverProperty.Precision);
+							}
+
+							if (bApplyRotation)
+							{
+								FQuat NewRotation = EffectorTarget->Rotation;
+								FTransform NewTransform = LinkData[EffectorLinkIndex].GetTransform();
+								NewTransform.SetRotation(NewRotation);
+								LinkData[EffectorLinkIndex].SetTransform(NewTransform);
+							}
+						}
+					}
+				}
+			}
+		}
+
+		///////////////////////////////////////////////////////////////////////////
+		// debug draw start
+		///////////////////////////////////////////////////////////////////////////
+		if (bDebugEnabled && Context.DrawInterface != nullptr)
+		{
+			const int32 DebugDataNum = DebugData.Num();
+			if (DebugData.Num() > 0)
+			{
+				for (int32 DebugIndex = DebugDataNum - 1; DebugIndex >= 0; --DebugIndex)
+				{
+					const TArray<FFBIKLinkData>& LocalLink = DebugData[DebugIndex].LinkData;
+
+					FTransform Offset = DebugOption.DrawWorldOffset;
+					Offset.SetLocation(Offset.GetLocation() * (DebugDataNum - DebugIndex));
+
+					if (DebugOption.bDrawDebugHierarchy)
+					{
+						for (int32 LinkIndex = 0; LinkIndex < LocalLink.Num(); ++LinkIndex)
+						{
+							const FFBIKLinkData& Data = LocalLink[LinkIndex];
+
+							FLinearColor DrawColor = FLinearColor::White;
+
+							float LineThickness = 0.f;
+							if (DebugOption.bColorAngularMotionStrength || DebugOption.bColorLinearMotionStrength)
+							{
+								DrawColor = FLinearColor::Black;
+								if (DebugOption.bColorAngularMotionStrength)
+								{
+									const float Range = FMath::Max(SolverProperty.AngularMotionStrength - SolverProperty.MinAngularMotionStrength, 0.f);
+									if (Range > 0.f)
+									{
+										float CurrentStrength = Data.GetAngularMotionStrength() - SolverProperty.MinAngularMotionStrength;
+										float Alpha = FMath::Clamp(CurrentStrength / Range, 0.f, 1.f);
+										DrawColor.R = LineThickness = Alpha;
+									}
+								}
+								else if (DebugOption.bColorLinearMotionStrength)
+								{
+									const float Range = FMath::Max(SolverProperty.LinearMotionStrength - SolverProperty.MinLinearMotionStrength, 0.f);
+									if (Range > 0.f)
+									{
+										float CurrentStrength = Data.GetLinearMotionStrength() - SolverProperty.MinLinearMotionStrength;
+										float Alpha = FMath::Clamp(CurrentStrength / Range, 0.f, 1.f);
+										DrawColor.B = LineThickness = Alpha;
+									}
+								}
+							}
+
+							if (Data.ParentLinkIndex != INDEX_NONE)
+							{
+								const FFBIKLinkData& ParentData = LocalLink[Data.ParentLinkIndex];
+								Context.DrawInterface->DrawLine(Offset, Data.GetPreviousTransform().GetLocation(), ParentData.GetPreviousTransform().GetLocation(), DrawColor, LineThickness);
+							}
+
+							if (DebugOption.bDrawDebugAxes)
+							{
+								Context.DrawInterface->DrawAxes(Offset, Data.GetPreviousTransform(), DebugOption.DrawSize);
+							}
+						}
+					}
+
+					if (DebugOption.bDrawDebugEffector)
+					{
+						for (auto Iter = EffectorTargets.CreateConstIterator(); Iter; ++Iter)
+						{
+							const FFBIKEffectorTarget& EffectorTarget = Iter.Value();
+							if (EffectorTarget.bPositionEnabled)
+							{
+								// draw effector target locations
+								Context.DrawInterface->DrawBox(Offset, FTransform(EffectorTarget.Position), FLinearColor::Yellow, DebugOption.DrawSize);
+							}
+
+							// draw effector link location
+							Context.DrawInterface->DrawBox(Offset, LocalLink[Iter.Key()].GetPreviousTransform(), FLinearColor::Green, DebugOption.DrawSize);
+						}
+
+						for (int32 Index = 0; Index < DebugData[DebugIndex].TargetVectorSources.Num(); ++Index)
+						{
+							// draw arrow to the target
+							Context.DrawInterface->DrawLine(Offset, DebugData[DebugIndex].TargetVectorSources[Index].GetLocation(),
+								DebugData[DebugIndex].TargetVectorSources[Index].GetLocation() + DebugData[DebugIndex].TargetVectors[Index], FLinearColor::Red);
+						}
+					}
+				}
+			}
+
+			if (DebugOption.bDrawDebugConstraints && InternalConstraints.Num())
+			{
+				FTransform Offset = FTransform::Identity;
+
+				// draw frame if active
+				for (int32 Index = 0; Index < Constraints.Num(); ++Index)
+				{
+					if (Constraints[Index].bEnabled)
+					{
+						if (Constraints[Index].Item.IsValid())
+						{
+							const int32* Found = HierarchyToLinkDataMap.Find(Constraints[Index].Item);
+							if (Found)
+							{
+								FTransform ConstraintFrame = LinkData[*Found].GetTransform();
+								ConstraintFrame.ConcatenateRotation(FQuat(Constraints[Index].OffsetRotation));
+								Context.DrawInterface->DrawAxes(Offset, ConstraintFrame, 2.f);
+							}
+						}
+					}
+				}
+
+				for (int32 Index = 0; Index < InternalConstraints.Num(); ++Index)
+				{
+					// for now we have rotation limit only
+					if (InternalConstraints[Index].IsType< FRotationLimitConstraint>())
+					{
+						FRotationLimitConstraint& LimitConstraint = InternalConstraints[Index].Get<FRotationLimitConstraint>();
+						const FQuat LocalRefRotation = LimitConstraint.RelativelRefPose.GetRotation();
+						FTransform RotationTransform = LinkData[LimitConstraint.ConstrainedIndex].GetTransform();
+						// base is parent transform but in their space, we can get there by inversing local ref rotation
+						FTransform BaseTransform = FTransform(LocalRefRotation).GetRelativeTransformReverse(RotationTransform);
+						BaseTransform.ConcatenateRotation(LimitConstraint.BaseFrameOffset);
+						Context.DrawInterface->DrawAxes(Offset, BaseTransform, 5.f, 1.f);
+
+						// current transform
+						const FQuat LocalRotation = BaseTransform.GetRotation().Inverse() * RotationTransform.GetRotation();
+						const FQuat DeltaTransform = LocalRefRotation.Inverse() * LocalRotation;
+						RotationTransform.SetRotation(BaseTransform.GetRotation() * DeltaTransform);
+						RotationTransform.NormalizeRotation();
+						RotationTransform.SetLocation(BaseTransform.GetLocation());
+
+						// draw ref pose on their current transform
+						Context.DrawInterface->DrawAxes(Offset, RotationTransform, 10.f, 1.f);
+
+						FVector XAxis = BaseTransform.GetUnitAxis(EAxis::X);
+						FVector YAxis = BaseTransform.GetUnitAxis(EAxis::Y);
+						FVector ZAxis = BaseTransform.GetUnitAxis(EAxis::Z);
+
+						if (LimitConstraint.bXLimitSet)
+						{
+							FTransform XAxisConeTM(YAxis, XAxis ^ YAxis, XAxis, BaseTransform.GetTranslation());
+							XAxisConeTM.SetRotation(FQuat(XAxis, 0.f) * XAxisConeTM.GetRotation());
+							XAxisConeTM.SetScale3D(FVector(30.f));
+							Context.DrawInterface->DrawCone(Offset, XAxisConeTM, LimitConstraint.Limit.X, 0.0f, 24, false, FLinearColor::Red, GEngine->ConstraintLimitMaterialX->GetRenderProxy());
+						}
+
+						if (LimitConstraint.bYLimitSet)
+						{
+							FTransform YAxisConeTM(ZAxis, YAxis ^ ZAxis, YAxis, BaseTransform.GetTranslation());
+							YAxisConeTM.SetRotation(FQuat(YAxis, 0.f) * YAxisConeTM.GetRotation());
+							YAxisConeTM.SetScale3D(FVector(30.f));
+							Context.DrawInterface->DrawCone(Offset, YAxisConeTM, LimitConstraint.Limit.Y, 0.0f, 24, false, FLinearColor::Green, GEngine->ConstraintLimitMaterialY->GetRenderProxy());
+						}
+
+						if (LimitConstraint.bZLimitSet)
+						{
+							FTransform ZAxisConeTM(XAxis, ZAxis ^ XAxis, ZAxis, BaseTransform.GetTranslation());
+							ZAxisConeTM.SetRotation(FQuat(ZAxis, 0.f) * ZAxisConeTM.GetRotation());
+							ZAxisConeTM.SetScale3D(FVector(30.f));
+							Context.DrawInterface->DrawCone(Offset, ZAxisConeTM, LimitConstraint.Limit.Z, 0.0f, 24, false, FLinearColor::Blue, GEngine->ConstraintLimitMaterialZ->GetRenderProxy());
+						}
+					}
+
+					if (InternalConstraints[Index].IsType<FPoleVectorConstraint>())
+					{
+						// darw pole vector location
+						// draw 3 joints line and a plane to pole vector
+						FPoleVectorConstraint& Constraint = InternalConstraints[Index].Get<FPoleVectorConstraint>();
+						FTransform RootTransform = LinkData[Constraint.ParentBoneIndex].GetTransform();
+						FTransform JointTransform = LinkData[Constraint.BoneIndex].GetTransform();
+						FTransform ChildTransform = LinkData[Constraint.ChildBoneIndex].GetTransform();
+
+						FVector JointTarget = (Constraint.bUseLocalDir) ? Constraint.CalculateCurrentPoleVectorDir(RootTransform, JointTransform, ChildTransform, LinkData[Constraint.BoneIndex].LocalFrame) : Constraint.PoleVector;
+
+						// draw the plane, 
+						TArray<FVector> Positions;
+						Positions.Add(RootTransform.GetLocation());
+						Positions.Add(ChildTransform.GetLocation());
+						Positions.Add(ChildTransform.GetLocation());
+						Positions.Add(JointTarget);
+						Positions.Add(JointTarget);
+						Positions.Add(RootTransform.GetLocation());
+
+						Context.DrawInterface->DrawLines(Offset, Positions, FLinearColor::Gray, 1.2f);
+						Context.DrawInterface->DrawLine(Offset, JointTransform.GetLocation(), JointTarget, FLinearColor::Red, 1.2f);
+					}
+				}
+			}
+		}
+		///////////////////////////////////////////////////////////////////////////
+		// debug draw end
+		///////////////////////////////////////////////////////////////////////////
+	}
+
+	// we update back to hierarchy
+	for (int32 LinkIndex = 0; LinkIndex < LinkData.Num(); ++LinkIndex)
+	{
+		// only propagate, if you are leaf joints here
+		// this means, only the last joint in the test
+		const FRigElementKey& CurrentItem = *LinkDataToHierarchyIndices.Find(LinkIndex);
+		const FTransform& LinkTransform = LinkData[LinkIndex].GetTransform();
+		Hierarchy->SetGlobalTransform(CurrentItem, LinkTransform, bPropagateToChildren);
+	}
+
+	#endif // 0 
+}
\ No newline at end of file
diff --git a/Engine/Plugins/Experimental/FullBodyIK/Source/FullBodyIK/Private/FullBodyIKSolver.h b/Engine/Plugins/Experimental/FullBodyIK/Source/FullBodyIK/Private/FullBodyIKSolver.h
new file mode 100644
--- /dev/null
+++ b/Engine/Plugins/Experimental/FullBodyIK/Source/FullBodyIK/Private/FullBodyIKSolver.h
@@ -0,0 +1,46 @@
+// Copyright Epic Games, Inc. All Rights Reserved.
+
+/**
+ * Contains Transform Solver Execution
+ *
+ */
+
+#pragma once
+
+#include "IKRigSolver.h"
+#include "FBIKShared.h"
+#include "FullBodyIKSolver.generated.h"
+
+struct FFBIKLinkData;
+
+ // run time for UFullBodyIKSolverDefinition
+UCLASS()
+class FULLBODYIK_API UFullBodyIKSolver : public UIKRigSolver
+{
+	GENERATED_BODY()
+
+	/** list of Link Data for solvers - joints */
+	TArray<FFBIKLinkData> LinkData;
+	/** Effector Targets - search key is LinkData Index */
+	TMap<int32, FFBIKEffectorTarget> EffectorTargets; 
+	/** End Effector Link Indices - EndEffector index to LinkData index*/
+	TArray<int32> EffectorLinkIndices;
+	/** Map from LinkData index to Rig Hierarchy Index*/
+	TMap<int32, FRigElementKey> LinkDataToHierarchyIndices;
+	/** Map from Rig Hierarchy Index to LinkData index*/
+	TMap<FRigElementKey, int32> HierarchyToLinkDataMap;
+	/** Constraints data */
+	TArray<ConstraintType> InternalConstraints;
+	/* Current Solver */
+	FJacobianSolver_FullbodyIK IKSolver;
+	/** Debug Data */
+	TArray<FJacobianDebugData> DebugData;
+
+protected:
+	virtual void InitInternal() override;
+	virtual void SolveInternal(FIKRigTransformModifier& InOutGlobalTransform) override;
+	virtual bool IsSolverActive() const override;
+};
+
+
+
diff --git a/Engine/Plugins/Experimental/FullBodyIK/Source/FullBodyIK/Private/FullBodyIKSolverDefinition.cpp b/Engine/Plugins/Experimental/FullBodyIK/Source/FullBodyIK/Private/FullBodyIKSolverDefinition.cpp
new file mode 100644
--- /dev/null
+++ b/Engine/Plugins/Experimental/FullBodyIK/Source/FullBodyIK/Private/FullBodyIKSolverDefinition.cpp
@@ -0,0 +1,72 @@
+// Copyright Epic Games, Inc. All Rights Reserved.
+
+/**
+ * Contains FullBodyIKSolverDefinition Implementation
+ *
+ */
+
+#include "FullBodyIKSolverDefinition.h"
+#include "FullBodyIKSolver.h"
+
+const FName UFullBodyIKSolverDefinition::EffectorTargetPrefix = FName(TEXT("FullBodyIKTarget"));
+
+UFullBodyIKSolverDefinition::UFullBodyIKSolverDefinition()
+{
+	DisplayName = TEXT("FullBodyIK Solver");
+	ExecutionClass = UFullBodyIKSolver::StaticClass();
+}
+
+#if WITH_EDITOR
+void UFullBodyIKSolverDefinition::UpdateEffectors()
+{
+	// first ensure we add all the effectors
+	for (int32 Index = 0; Index < Effectors.Num(); ++Index)
+	{
+		EnsureToAddEffector(Effectors[Index].Target, TEXT("FBIK_Effector"));
+	}
+
+	// we have more nodes than goals
+	// which means we have something deleted
+	if (Effectors.Num() < EffectorToGoal.Num())
+	{
+		TArray<FIKRigEffector> GoalEffectors;
+		// we have to remove things that don't belong
+		for (auto Iter = EffectorToGoal.CreateIterator(); Iter; ++Iter)
+		{
+			GoalEffectors.Add(Iter.Key());
+		}
+
+		TBitArray<> RemoveFlags(false, GoalEffectors.Num());
+
+		for (int32 Index = 0; Index < Effectors.Num(); ++Index)
+		{
+			int32 Found = GoalEffectors.Find(Effectors[Index].Target);
+
+			if (Found != INDEX_NONE)
+			{
+				RemoveFlags[Found] = true;
+			}
+		}
+
+		for (TConstSetBitIterator<> Iter(RemoveFlags); Iter; ++Iter)
+		{
+			// remove things that don't belong
+			EffectorToGoal.Remove(GoalEffectors[Iter.GetIndex()]);
+		}
+	}
+
+ 	// trigger a delegate for goal has been updated?
+ 	OnGoalHasBeenUpdated();
+}
+
+void UFullBodyIKSolverDefinition::PostEditChangeChainProperty(struct FPropertyChangedChainEvent& PropertyChangedEvent)
+{
+	if (PropertyChangedEvent.Property->GetFName() == GET_MEMBER_NAME_CHECKED(UFullBodyIKSolverDefinition, Effectors))
+	{
+		UpdateEffectors();
+	}
+
+	Super::PostEditChangeChainProperty(PropertyChangedEvent);
+}
+
+#endif// WITH_EDITOR
\ No newline at end of file
diff --git a/Engine/Plugins/Experimental/FullBodyIK/Source/FullBodyIK/Private/FullBodyIKSolverDefinition.h b/Engine/Plugins/Experimental/FullBodyIK/Source/FullBodyIK/Private/FullBodyIKSolverDefinition.h
new file mode 100644
--- /dev/null
+++ b/Engine/Plugins/Experimental/FullBodyIK/Source/FullBodyIK/Private/FullBodyIKSolverDefinition.h
@@ -0,0 +1,77 @@
+// Copyright Epic Games, Inc. All Rights Reserved.
+
+/**
+ * Contains Fullbody IK solver Definition
+ *
+ */
+
+#pragma once
+
+#include "IKRigSolverDefinition.h"
+#include "FBIKShared.h"
+#include "FBIKConstraintOption.h"
+#include "FullBodyIKSolverDefinition.generated.h"
+
+USTRUCT()
+struct FFBIKRigEffector
+{
+	GENERATED_BODY()
+	/**
+	 * The last item in the chain to solve - the effector
+	 */
+	UPROPERTY(EditAnywhere, Category = FFBIKRigEffector)
+	FIKRigEffector Target;
+
+	UPROPERTY(EditAnywhere, Category = FFBIKRigEffector)
+	int32  PositionDepth = 1000;
+	
+	UPROPERTY(EditAnywhere, Category = FFBIKRigEffector)
+	int32  RotationDepth = 1000;
+
+	/*
+	 * Clamps the total length to target by this scale for each iteration 
+	 * This helps to stabilize solver to reduce singularity by avoiding to try to reach target too far. 
+	 */
+	UPROPERTY(EditAnywhere, Category = FFBIKRigEffector)
+	float	Pull = 0.f;
+
+	FFBIKRigEffector()
+	{
+	}
+};
+
+UCLASS(BlueprintType)
+class FULLBODYIK_API UFullBodyIKSolverDefinition : public UIKRigSolverDefinition
+{
+	GENERATED_BODY()
+
+public:
+	UFullBodyIKSolverDefinition();
+
+	UPROPERTY(EditAnywhere, Category = Effectors)
+	FName Root;
+
+	UPROPERTY(EditAnywhere, Category = Effectors)
+	TArray<FFBIKRigEffector> Effectors;
+//	UPROPERTY(EditAnywhere, Category = Constraint)
+//	TArray<FFBIKConstraintOption> Constraints;
+
+	UPROPERTY(EditAnywhere, Category = Solver)
+	FSolverInput	SolverProperty;
+
+	UPROPERTY(EditAnywhere, Category = Motion)
+	FMotionProcessInput MotionProperty;
+
+private:
+#if WITH_EDITOR
+	virtual void UpdateEffectors() override;
+
+	// UObject interface
+	virtual void PostEditChangeChainProperty(struct FPropertyChangedChainEvent& PropertyChangedEvent) override;
+	// UObject interface
+#endif // WITH_EDITOR
+
+public:
+	static const FName EffectorTargetPrefix;
+};
+
diff --git a/Engine/Plugins/Experimental/FullBodyIK/Source/FullBodyIK/Private/RigUnit_FullbodyIK.cpp b/Engine/Plugins/Experimental/FullBodyIK/Source/FullBodyIK/Private/RigUnit_FullbodyIK.cpp
--- a/Engine/Plugins/Experimental/FullBodyIK/Source/FullBodyIK/Private/RigUnit_FullbodyIK.cpp
+++ b/Engine/Plugins/Experimental/FullBodyIK/Source/FullBodyIK/Private/RigUnit_FullbodyIK.cpp
@@ -7,51 +7,6 @@
 #include "FBIKConstraintLib.h"
 
 #define MAX_DEPTH 10000
-
-void FJacobianSolver_FullbodyIK::InitializeSolver(TArray<FFBIKLinkData>& InOutLinkData, TMap<int32, FFBIKEffectorTarget>& InOutEffectors) const
-{
-// since we're using constraints, we don't want to create new motion base all the time
-// constraints will add that info and we'll utilize it
-}
-
-void FJacobianSolver_FullbodyIK::PreSolve(TArray<FFBIKLinkData>& InOutLinkData, const TMap<int32, FFBIKEffectorTarget>& InEffectors) const
-{
-	// for final one, we'll have to add those in the beginning, and then just update the frame for every iteration
-	for (int32 Index = 0; Index < InOutLinkData.Num(); ++Index)
-	{
-		// we update existing data
-		const FQuat& LocalFrame = InOutLinkData[Index].LocalFrame;
-
-		FQuat CurrentFrame;
-
-		if (InOutLinkData[Index].ParentLinkIndex != INDEX_NONE)
-		{
-			CurrentFrame = InOutLinkData[InOutLinkData[Index].ParentLinkIndex].GetTransform().GetRotation() * LocalFrame;
-		}
-		else
-		{
-			CurrentFrame = LocalFrame;
-		}
-
-		// update the base axis
-		if (InOutLinkData[Index].GetNumMotionBases() == 3)
-		{
-			InOutLinkData[Index].GetMotionBase(0).BaseAxis = CurrentFrame.RotateVector(FVector::ForwardVector);
-			InOutLinkData[Index].GetMotionBase(1).BaseAxis = CurrentFrame.RotateVector(FVector::RightVector);
-			InOutLinkData[Index].GetMotionBase(2).BaseAxis = CurrentFrame.RotateVector(FVector::UpVector);
-		}
-		else
-		{
-			ensure(InOutLinkData[Index].GetNumMotionBases() == 0);
-
-			// invalid number of motion bases for this solver
-			InOutLinkData[Index].ResetMotionBases();
-			InOutLinkData[Index].AddMotionBase(FMotionBase(CurrentFrame.RotateVector(FVector::ForwardVector)));
-			InOutLinkData[Index].AddMotionBase(FMotionBase(CurrentFrame.RotateVector(FVector::RightVector)));
-			InOutLinkData[Index].AddMotionBase(FMotionBase(CurrentFrame.RotateVector(FVector::UpVector)));
-		}
-	}
-}
 /////////////////////////////////////////////////////
 
 /////////////////////////////////////////////////////
diff --git a/Engine/Plugins/Experimental/FullBodyIK/Source/FullBodyIK/Private/RigUnit_FullbodyIK.h b/Engine/Plugins/Experimental/FullBodyIK/Source/FullBodyIK/Private/RigUnit_FullbodyIK.h
--- a/Engine/Plugins/Experimental/FullBodyIK/Source/FullBodyIK/Private/RigUnit_FullbodyIK.h
+++ b/Engine/Plugins/Experimental/FullBodyIK/Source/FullBodyIK/Private/RigUnit_FullbodyIK.h
@@ -3,9 +3,9 @@
 #pragma once
 
 #include "Units/Highlevel/RigUnit_HighlevelBase.h"
-#include "JacobianSolver.h"
 #include "FBIKConstraintOption.h"
 #include "FBIKDebugOption.h"
+#include "FBIKShared.h"
 #include "RigUnit_FullbodyIK.generated.h"
 
 USTRUCT()
@@ -69,17 +69,6 @@ struct FFBIKEndEffector
 	}
 };
 
-class FULLBODYIK_API FJacobianSolver_FullbodyIK : public FJacobianSolverBase
-{
-public:
-	FJacobianSolver_FullbodyIK() : FJacobianSolverBase() {}
-
-private:
-	virtual void InitializeSolver(TArray<FFBIKLinkData>& InOutLinkData, TMap<int32, FFBIKEffectorTarget>& InOutEndEffectors) const override;
-	virtual void PreSolve(TArray<FFBIKLinkData>& InOutLinkData, const TMap<int32, FFBIKEffectorTarget>& InEndEffectors) const override;
-};
-
-
 USTRUCT()
 struct FRigUnit_FullbodyIK_WorkData
 {
@@ -107,63 +96,6 @@ struct FRigUnit_FullbodyIK_WorkData
 	TArray<FJacobianDebugData> DebugData;
 };
 
-USTRUCT()
-struct FSolverInput
-{
-	GENERATED_BODY()
-
-	/*
-	 * This value is applied to the target information for effectors, which influence back to 
-	 * Joint's motion that are affected by the end effector
-	 * The reason min/max is used when we apply the depth through the chain that are affected
-
-	 */
-	UPROPERTY()
-	float	LinearMotionStrength = 3.f;
-
-	UPROPERTY()
-	float	MinLinearMotionStrength = 2.f;
-
-	/*
-	 * This value is applied to the target information for effectors, which influence back to 
-	 * Joint's motion that are affected by the end effector
-	 * The reason min/max is used when we apply the depth through the chain that are affected
-	 */
-	UPROPERTY()
-	float	AngularMotionStrength = 3.f;
-
-	UPROPERTY()
-	float	MinAngularMotionStrength = 2.f;
-
-	/* This is a scale value (range from 0-0.7) that is used to stablize the target vector. If less, it's more stable, but it can reduce speed of converge. */
-	UPROPERTY()
-	float	DefaultTargetClamp = 0.2f;
-
-	/**
-	 * The precision to use for the solver
-	 */
-	UPROPERTY()
-	float Precision = 0.1f;
-
-	/**
-	* The precision to use for the fabrik solver
-	*/
-	UPROPERTY()
-	float Damping = 30.f;
-
-	/**
-	 * The maximum number of iterations. Values between 4 and 16 are common.
-	 */
-	UPROPERTY()
-	int32 MaxIterations = 30;
-
-	/**
-	 * Cheaper solution than default Jacobian Pseudo Inverse Damped Least Square
-	 */
-	UPROPERTY()
-	bool bUseJacobianTranspose = false;
-};
-
 /**
  * Based on Jacobian solver at core, this can solve multi chains within a root using multi effectors
  */
