- Change IKSolver init to support transform modifier
- Added FBIK solver runtime code
- Added an option to debug draw

diff --git a/Engine/Plugins/Animation/IKRig/IKRig.uplugin b/Engine/Plugins/Animation/IKRig/IKRig.uplugin
--- a/Engine/Plugins/Animation/IKRig/IKRig.uplugin
+++ b/Engine/Plugins/Animation/IKRig/IKRig.uplugin
@@ -15,6 +15,12 @@
 	"IsBetaVersion": true,
 	"IsExperimentalVersion": false,
 	"Installed": false,
+	"Plugins": [
+		{
+			"Name": "ControlRig",
+			"Enabled": true
+		}
+	],
 	"Modules": [
 		{
 			"Name": "IKRig",
diff --git a/Engine/Plugins/Animation/IKRig/Source/IKRig/IKRig.Build.cs b/Engine/Plugins/Animation/IKRig/Source/IKRig/IKRig.Build.cs
--- a/Engine/Plugins/Animation/IKRig/Source/IKRig/IKRig.Build.cs
+++ b/Engine/Plugins/Animation/IKRig/Source/IKRig/IKRig.Build.cs
@@ -16,7 +16,8 @@ public class IKRig : ModuleRules
 				"Slate",
 				"SlateCore",
 				"AnimationCore", 
-				"AnimGraphRuntime"
+				"AnimGraphRuntime",
+				"ControlRig"
 			}
 			);
 
diff --git a/Engine/Plugins/Animation/IKRig/Source/IKRig/Private/AnimNodes/AnimNode_IKRig.cpp b/Engine/Plugins/Animation/IKRig/Source/IKRig/Private/AnimNodes/AnimNode_IKRig.cpp
--- a/Engine/Plugins/Animation/IKRig/Source/IKRig/Private/AnimNodes/AnimNode_IKRig.cpp
+++ b/Engine/Plugins/Animation/IKRig/Source/IKRig/Private/AnimNodes/AnimNode_IKRig.cpp
@@ -66,23 +66,20 @@ void FAnimNode_IKRig::Evaluate_AnyThread(FPoseContext& Output)
 		// fill up current pose
 		RigProcessor->Solve();
 
-		// get component pose from control rig
-//		FCSPose<FCompactPose> MeshPoses;
-//		MeshPoses.InitPose(SourcePose.Pose);
-
 		// now we copy back to output pose
 		for (FCompactPoseBoneIndex CPIndex : OutPose.ForEachBoneIndex())
 		{
 			int32* Index = CompactPoseToRigIndices.Find(CPIndex);
 			if (Index)
 			{
-//				MeshPoses.SetComponentSpaceTransform(CPIndex, TransformModifier.GetGlobalTransform(*Index));
 				Output.Pose[CPIndex] = TransformModifier.GetLocalTransform(*Index);
 			}
 		}
 
-//		FCSPose<FCompactPose>::ConvertComponentPosesToLocalPosesSafe(MeshPoses, Output.Pose);
 		Output.Pose.NormalizeRotations();
+
+		// draw the interface
+		QueueDrawInterface(Output.AnimInstanceProxy, Output.AnimInstanceProxy->GetComponentTransform());
 	}
 	else
 	{
@@ -201,4 +198,63 @@ void FAnimNode_IKRig::CacheBones_AnyThread(const FAnimationCacheBonesContext& Co
 			}
 		}
 	}
+}
+
+void FAnimNode_IKRig::QueueDrawInterface(FAnimInstanceProxy* AnimProxy, const FTransform& ComponentToWorld)
+{
+	check (RigProcessor);
+	check (AnimProxy);
+
+	for (const FControlRigDrawInstruction& Instruction : RigProcessor->GetDrawInterface())
+	{
+		if (!Instruction.IsValid())
+		{
+			continue;
+		}
+
+		FTransform InstructionTransform = Instruction.Transform * ComponentToWorld;
+		switch (Instruction.PrimitiveType)
+		{
+		case EControlRigDrawSettings::Points:
+		{
+			ensure(false);
+			// no support for points yet, but feel free to add that to AnimProxy
+// 			for (const FVector& Point : Instruction.Positions)
+// 			{
+// 				AnimProxy->Add (InstructionTransform.TransformPosition(Point), Instruction.Color, Instruction.Thickness, SDPG_Foreground);
+// 			}
+ 			break;
+		}
+		case EControlRigDrawSettings::Lines:
+		{
+			const TArray<FVector>& Points = Instruction.Positions;
+			for (int32 PointIndex = 0; PointIndex < Points.Num() - 1; PointIndex += 2)
+			{
+				AnimProxy->AnimDrawDebugLine(InstructionTransform.TransformPosition(Points[PointIndex]), InstructionTransform.TransformPosition(Points[PointIndex + 1]), Instruction.Color.ToFColor(false), false, 0.f, Instruction.Thickness);
+			}
+			break;
+		}
+		case EControlRigDrawSettings::LineStrip:
+		{
+			const TArray<FVector>& Points = Instruction.Positions;
+			for (int32 PointIndex = 0; PointIndex < Points.Num() - 1; PointIndex++)
+			{
+				AnimProxy->AnimDrawDebugLine(InstructionTransform.TransformPosition(Points[PointIndex]), InstructionTransform.TransformPosition(Points[PointIndex + 1]), Instruction.Color.ToFColor(false), false, 0.f, Instruction.Thickness);
+			}
+			break;
+		}
+
+		case EControlRigDrawSettings::DynamicMesh:
+		{
+			ensure(false);
+			// no support for this yet
+// 			FDynamicMeshBuilder MeshBuilder(PDI->View->GetFeatureLevel());
+// 			MeshBuilder.AddVertices(Instruction.MeshVerts);
+// 			MeshBuilder.AddTriangles(Instruction.MeshIndices);
+// 			MeshBuilder.Draw(PDI, InstructionTransform.ToMatrixWithScale(), Instruction.MaterialRenderProxy, SDPG_World/*SDPG_Foreground*/);
+			break;
+		}
+
+		}
+	}
 }
\ No newline at end of file
diff --git a/Engine/Plugins/Animation/IKRig/Source/IKRig/Private/IKRigDefinition.cpp b/Engine/Plugins/Animation/IKRig/Source/IKRig/Private/IKRigDefinition.cpp
--- a/Engine/Plugins/Animation/IKRig/Source/IKRig/Private/IKRigDefinition.cpp
+++ b/Engine/Plugins/Animation/IKRig/Source/IKRig/Private/IKRigDefinition.cpp
@@ -301,7 +301,7 @@ void UIKRigDefinition::EnsureSortedCorrectly(bool bReSortIfNeeded)
 						}
 
 						Iter = Iter->Sibling;
-					} while (Iter);
+					} while (Iter != nullptr);
 				}
 			}
 		};
diff --git a/Engine/Plugins/Animation/IKRig/Source/IKRig/Private/IKRigProcessor.cpp b/Engine/Plugins/Animation/IKRig/Source/IKRig/Private/IKRigProcessor.cpp
--- a/Engine/Plugins/Animation/IKRig/Source/IKRig/Private/IKRigProcessor.cpp
+++ b/Engine/Plugins/Animation/IKRig/Source/IKRig/Private/IKRigProcessor.cpp
@@ -20,10 +20,9 @@ void UIKRigProcessor::SetIKRigDefinition(UIKRigDefinition* InRigDefinition, bool
 
 	RigDefinition = InRigDefinition;
 
-	TransformModifier = FIKRigTransformModifier(&RigDefinition->GetHierarchy());
-
 	if (RigDefinition && bInitialize)
 	{
+		TransformModifier = FIKRigTransformModifier(&RigDefinition->GetHierarchy());
 		Initialize(RigDefinition->GetReferencePose());
 	}
 }
@@ -57,7 +56,8 @@ void UIKRigProcessor::Reinitialize()
 			if (ClassType != nullptr)
 			{
 				UIKRigSolver* NewSolver = NewObject<UIKRigSolver>(this, ClassType);
-				NewSolver->Init(SolverDefinitions[Index], UIKRigSolver::FIKRigTransformGetter::CreateUObject(this, &UIKRigProcessor::GetRefPoseGetter), 
+				// const transform modifier for hierarchy query
+				NewSolver->Init(SolverDefinitions[Index], TransformModifier, UIKRigSolver::FIKRigTransformGetter::CreateUObject(this, &UIKRigProcessor::GetRefPoseGetter),
 					UIKRigSolver::FIKRigGoalGetter::CreateUObject(this, &UIKRigProcessor::GoalGetter));
 				Solvers.Add(NewSolver);
 			}
@@ -96,9 +96,11 @@ void UIKRigProcessor::Solve()
 	{
 		check (RigDefinition);
 
+		DrawInterface.Reset();
 		for (int32 Index = 0; Index < Solvers.Num(); ++Index)
 		{
-			Solvers[Index]->Solve(TransformModifier);
+			// Draw interface is pointer in case we want to nullify for optimization
+			Solvers[Index]->Solve(TransformModifier, &DrawInterface);
 		}
 	}
 }
diff --git a/Engine/Plugins/Animation/IKRig/Source/IKRig/Private/IKRigSolver.cpp b/Engine/Plugins/Animation/IKRig/Source/IKRig/Private/IKRigSolver.cpp
--- a/Engine/Plugins/Animation/IKRig/Source/IKRig/Private/IKRigSolver.cpp
+++ b/Engine/Plugins/Animation/IKRig/Source/IKRig/Private/IKRigSolver.cpp
@@ -9,7 +9,7 @@
 #include "IKRigDataTypes.h"
 
 // input hierarchy and ref pose? 
-void UIKRigSolver::Init(UIKRigSolverDefinition* InSolverDefinition, FIKRigTransformGetter InRefPoseGetter, FIKRigGoalGetter InGoalGetter)
+void UIKRigSolver::Init(UIKRigSolverDefinition* InSolverDefinition, const FIKRigTransformModifier& TransformModifier, FIKRigTransformGetter InRefPoseGetter, FIKRigGoalGetter InGoalGetter)
 {
 	SolverDefinition = InSolverDefinition;
 
@@ -19,7 +19,7 @@ void UIKRigSolver::Init(UIKRigSolverDefinition* InSolverDefinition, FIKRigTransf
 	GoalGetter = InGoalGetter;
 	ensure(GoalGetter.IsBound());
 
-	InitInternal();
+	InitInternal(TransformModifier);
 }
 
 bool UIKRigSolver::IsSolverActive() const 
@@ -29,11 +29,11 @@ bool UIKRigSolver::IsSolverActive() const
 
 // input : goal getter or goals
 // output : modified pose - GlobalTransforms
-void UIKRigSolver::Solve(FIKRigTransformModifier& InOutGlobalTransform)
+void UIKRigSolver::Solve(FIKRigTransformModifier& InOutGlobalTransform, FControlRigDrawInterface* InOutDrawInterface)
 {
 	if (IsSolverActive())
 	{
-		SolveInternal(InOutGlobalTransform);
+		SolveInternal(InOutGlobalTransform, InOutDrawInterface);
 	}
 }
 
diff --git a/Engine/Plugins/Animation/IKRig/Source/IKRig/Private/Solvers/TransformSolver.cpp b/Engine/Plugins/Animation/IKRig/Source/IKRig/Private/Solvers/TransformSolver.cpp
--- a/Engine/Plugins/Animation/IKRig/Source/IKRig/Private/Solvers/TransformSolver.cpp
+++ b/Engine/Plugins/Animation/IKRig/Source/IKRig/Private/Solvers/TransformSolver.cpp
@@ -9,7 +9,7 @@
 #include "IKRigHierarchy.h"
 #include "Solvers/TransformSolverDefinition.h"
 
-void UTransformSolver::InitInternal()
+void UTransformSolver::InitInternal(const FIKRigTransformModifier& InGlobalTransform)
 {
 	
 }
@@ -27,7 +27,7 @@ bool UTransformSolver::IsSolverActive() const
 	return false;
 }
 
-void UTransformSolver::SolveInternal(FIKRigTransformModifier& InOutGlobalTransform)
+void UTransformSolver::SolveInternal(FIKRigTransformModifier& InOutGlobalTransform, FControlRigDrawInterface* InOutDrawInterface)
 {
 	if (UTransformSolverDefinition* SolverDef = Cast<UTransformSolverDefinition>(SolverDefinition))
 	{
diff --git a/Engine/Plugins/Animation/IKRig/Source/IKRig/Private/Solvers/TransformSolver.h b/Engine/Plugins/Animation/IKRig/Source/IKRig/Private/Solvers/TransformSolver.h
--- a/Engine/Plugins/Animation/IKRig/Source/IKRig/Private/Solvers/TransformSolver.h
+++ b/Engine/Plugins/Animation/IKRig/Source/IKRig/Private/Solvers/TransformSolver.h
@@ -18,8 +18,8 @@ class IKRIG_API UTransformSolver : public UIKRigSolver
 	GENERATED_BODY()
 
 protected:
-	virtual void InitInternal() override;
-	virtual void SolveInternal(FIKRigTransformModifier& InOutGlobalTransform) override;
+	virtual void InitInternal(const FIKRigTransformModifier& InGlobalTransform) override;
+	virtual void SolveInternal(FIKRigTransformModifier& InOutGlobalTransform, FControlRigDrawInterface* InOutDrawInterface) override;
 	virtual bool IsSolverActive() const override;
 };
 
diff --git a/Engine/Plugins/Animation/IKRig/Source/IKRig/Public/AnimNodes/AnimNode_IKRig.h b/Engine/Plugins/Animation/IKRig/Source/IKRig/Public/AnimNodes/AnimNode_IKRig.h
--- a/Engine/Plugins/Animation/IKRig/Source/IKRig/Public/AnimNodes/AnimNode_IKRig.h
+++ b/Engine/Plugins/Animation/IKRig/Source/IKRig/Public/AnimNodes/AnimNode_IKRig.h
@@ -44,7 +44,6 @@ private:
 	UIKRigProcessor* RigProcessor = nullptr;
 
 	TArray<FName> GoalNames;
-
 	TMap<FCompactPoseBoneIndex, int32, FDefaultSetAllocator, TCompactPoseBoneIndexMapKeyFuncs<int32>> CompactPoseToRigIndices;
 
 public:
@@ -64,5 +63,7 @@ private:
 	bool RebuildGoalList();
 	FName GetGoalName(int32 Index) const;
 
+	void QueueDrawInterface(FAnimInstanceProxy* AnimProxy, const FTransform& ComponentToWorld);
+
 	friend class UAnimGraphNode_IKRig;
 };
diff --git a/Engine/Plugins/Animation/IKRig/Source/IKRig/Public/IKRigProcessor.h b/Engine/Plugins/Animation/IKRig/Source/IKRig/Public/IKRigProcessor.h
--- a/Engine/Plugins/Animation/IKRig/Source/IKRig/Public/IKRigProcessor.h
+++ b/Engine/Plugins/Animation/IKRig/Source/IKRig/Public/IKRigProcessor.h
@@ -11,11 +11,13 @@
 #include "UObject/ObjectMacros.h"
 #include "UObject/Object.h"
 #include "IKRigDataTypes.h"
+#include "Drawing/ControlRigDrawInterface.h"
 #include "IKRigProcessor.generated.h"
 
 
 class UIKRigConstraintSolver;
 class UIKRigSolver;
+struct FControlRigDrawInterface;
 
 UCLASS(BlueprintType)
 class IKRIG_API UIKRigProcessor : public UObject
@@ -47,6 +49,7 @@ private:
 	TMap<FName, FIKRigGoal> IKGoals;
 
 	FIKRigTransformModifier TransformModifier;
+	FControlRigDrawInterface DrawInterface;
 
 	bool bInitialized = false;
 
@@ -70,6 +73,8 @@ public:
 	const FIKRigHierarchy* GetHierarchy() const;
 
 	void ResetToRefPose();
+	const FControlRigDrawInterface& GetDrawInterface() { return DrawInterface; }
+
 private:
 	FIKRigTarget* FindGoal(const FName& GoalName);
 
diff --git a/Engine/Plugins/Animation/IKRig/Source/IKRig/Public/IKRigSolver.h b/Engine/Plugins/Animation/IKRig/Source/IKRig/Public/IKRigSolver.h
--- a/Engine/Plugins/Animation/IKRig/Source/IKRig/Public/IKRigSolver.h
+++ b/Engine/Plugins/Animation/IKRig/Source/IKRig/Public/IKRigSolver.h
@@ -17,6 +17,7 @@ struct FIKRigTarget;
 struct FIKRigTransform;
 struct FIKRigTransformModifier;
 class UIKRigSolverDefinition;
+struct FControlRigDrawInterface;
 
 // run time processor 
 UCLASS(Abstract, hidecategories = UObject)
@@ -40,16 +41,17 @@ public:
 	bool bEnabled = true;
 
 	// input hierarchy and ref pose? 
-	void Init(UIKRigSolverDefinition* InSolverDefinition, FIKRigTransformGetter InRefPoseGetter, FIKRigGoalGetter InGoalGetter/*, FSolveConstraint& InConstraintHandler*/);
+	void Init(UIKRigSolverDefinition* InSolverDefinition, const FIKRigTransformModifier& TransformModifier, FIKRigTransformGetter InRefPoseGetter, FIKRigGoalGetter InGoalGetter/*, FSolveConstraint& InConstraintHandler*/);
 
 	// input : goal getter or goals
 	// output : modified pose - GlobalTransforms
 	// or use SolverInternal function
-	void Solve(FIKRigTransformModifier& InOutGlobalTransform);
+	void Solve(FIKRigTransformModifier& InOutGlobalTransform, FControlRigDrawInterface* InOutDrawInterface);
 
 protected:
-	virtual void InitInternal() {};
-	virtual void SolveInternal(FIKRigTransformModifier& InOutGlobalTransform) {};
+	// during we don't mutate bone transform, but you can read it
+	virtual void InitInternal(const FIKRigTransformModifier& InGlobalTransform) {};
+	virtual void SolveInternal(FIKRigTransformModifier& InOutGlobalTransform, FControlRigDrawInterface* InOutDrawInterface) {};
 	virtual bool IsSolverActive() const;
 
 	bool GetEffectorTarget(const FIKRigEffector& InEffector, FIKRigTarget& OutTarget) const;
diff --git a/Engine/Plugins/Experimental/FullBodyIK/Source/FullBodyIK/Private/FullBodyIKSolver.cpp b/Engine/Plugins/Experimental/FullBodyIK/Source/FullBodyIK/Private/FullBodyIKSolver.cpp
--- a/Engine/Plugins/Experimental/FullBodyIK/Source/FullBodyIK/Private/FullBodyIKSolver.cpp
+++ b/Engine/Plugins/Experimental/FullBodyIK/Source/FullBodyIK/Private/FullBodyIKSolver.cpp
@@ -10,19 +10,206 @@
 #include "FBIKConstraint.h"
 #include "JacobianIK.h"
 #include "FullBodyIKSolverDefinition.h"
+#include "FBIKUtil.h"
+#include "FBIKShared.h"
+#include "Drawing/ControlRigDrawInterface.h"
+
+//////////////////////////////////////////////
+// utility functions
+////////////////////////////////////////////////
+static float EnsureToAddBoneToLinkData(const FIKRigTransformModifier& TransformModifier, const int32 CurrentItem, TArray<FFBIKLinkData>& LinkData,
+	TMap<int32, int32>& HierarchyToLinkDataMap, TMap<int32, int32>& LinkDataToHierarchyIndices)
+{
+	float ChainLength = 0;
+
+	// we insert from back to first, but only if the list doesn't have it
+	int32* FoundLinkIndex = HierarchyToLinkDataMap.Find(CurrentItem);
+	if (!FoundLinkIndex)
+	{
+		int32 NewLinkIndex = LinkData.AddDefaulted();
+		FFBIKLinkData& NewLink = LinkData[NewLinkIndex];
+
+		// find parent LinkIndex
+		const int32 ParentItem = TransformModifier.Hierarchy->GetParentIndex(CurrentItem);
+		FoundLinkIndex = HierarchyToLinkDataMap.Find(ParentItem);
+		NewLink.ParentLinkIndex = (FoundLinkIndex) ? *FoundLinkIndex : INDEX_NONE;
+		NewLink.SetTransform(TransformModifier.GetGlobalTransform(CurrentItem));
+
+		if (ParentItem != INDEX_NONE)
+		{
+			// @todo: currently we assume the input pose is init transform
+			// this could use reference pose getter from IKRigSolver or we could wrap it to TransformModifier
+			FVector DiffLocation = TransformModifier.GetGlobalTransform(CurrentItem).GetLocation() - TransformModifier.GetGlobalTransform(ParentItem).GetLocation();
+			// set Length
+			NewLink.Length = DiffLocation.Size();
+		}
+
+		// we create bidirectional look up table
+		HierarchyToLinkDataMap.Add(CurrentItem, NewLinkIndex);
+		LinkDataToHierarchyIndices.Add(NewLinkIndex, CurrentItem);
 
-void UFullBodyIKSolver::InitInternal()
+		ChainLength = NewLink.Length;
+	}
+	else
+	{
+		ChainLength = LinkData[*FoundLinkIndex].Length;
+	}
+
+	return ChainLength;
+}
+
+static void AddToEffectorTarget(int32 EffectorIndex, const int32 Effector, TMap<int32, FFBIKEffectorTarget>& EffectorTargets, const TMap<int32, int32>& HierarchyToLinkDataMap,
+	TArray<int32>& EffectorLinkIndices, float ChainLength, const TArray<int32>& ChainIndices, int32 PositionDepth, int32 RotationDepth)
+{
+	const int32* EffectorLinkIndexID = HierarchyToLinkDataMap.Find(Effector);
+	check(EffectorLinkIndexID);
+	EffectorLinkIndices[EffectorIndex] = *EffectorLinkIndexID;
+	// add EffectorTarget for this link Index
+	FFBIKEffectorTarget& EffectorTarget = EffectorTargets.FindOrAdd(*EffectorLinkIndexID);
+	EffectorTarget.ChainLength = ChainLength;
+
+	// convert bone chain indices to link chain
+	const int32 MaxNum = FMath::Max(PositionDepth, RotationDepth);
+	if (ensure(MaxNum <= ChainIndices.Num()))
+	{
+		EffectorTarget.LinkChain.Reset(MaxNum);
+		for (int32 Index = 0; Index < MaxNum; ++Index)
+		{
+			const int32 Bone = ChainIndices[Index];
+			const int32* LinkIndexID = HierarchyToLinkDataMap.Find(Bone);
+			EffectorTarget.LinkChain.Add(*LinkIndexID);
+		}
+	}
+}
+
+// the output should be from current index -> to root parent 
+static bool GetBoneChain(const FIKRigTransformModifier& TransformModifier, const FName& Root, const FName& Current, TArray<int32>& ChainIndices)
+{
+	ChainIndices.Reset();
+
+	int32 RootIndex = TransformModifier.Hierarchy->GetIndex(Root);
+	int32 Iterator = TransformModifier.Hierarchy->GetIndex(Current);;
+
+	// iterates until key is valid
+	while (Iterator!=INDEX_NONE && Iterator != RootIndex)
+	{
+		ChainIndices.Insert(Iterator, 0);
+		Iterator = TransformModifier.Hierarchy->GetParentIndex(Iterator);
+	}
+
+	// add the last one if valid
+	if (Iterator!=INDEX_NONE)
+	{
+		ChainIndices.Insert(Iterator, 0);
+	}
+
+	// when you reached to something invalid, that means, we did not hit the 
+	// expected root, we iterated to the root and above, and we exhausted our option
+	// so if you hit or valid target by the hitting max depth, we should 
+	return ChainIndices.Num() > 0;
+}
+
+static void AddEffectors(const FIKRigTransformModifier& TransformModifier, const FName Root, const TArray<FFBIKRigEffector>& Effectors,
+	TArray<FFBIKLinkData>& LinkData, TMap<int32, FFBIKEffectorTarget>& EffectorTargets, TArray<int32>& EffectorLinkIndices,
+	TMap<int32, int32>& LinkDataToHierarchyIndices, TMap<int32, int32>& HierarchyToLinkDataMap, const FSolverInput& SolverProperty)
+{
+	EffectorLinkIndices.SetNum(Effectors.Num());
+	// fill up all effector indices
+	for (int32 Index = 0; Index < Effectors.Num(); ++Index)
+	{
+		// clear link indices, so that we don't search
+		EffectorLinkIndices[Index] = INDEX_NONE;
+		// create LinkeData from root bone to all effectors 
+		const FName Bone = Effectors[Index].Target.Bone;
+		const int32 BoneIndex = TransformModifier.Hierarchy->GetIndex(Bone);
+
+		if (Bone == NAME_None || BoneIndex == INDEX_NONE)
+		{
+			continue;
+		}
+
+		const FFBIKRigEffector& CurrentEffector = Effectors[Index];
+
+		TArray<int32> ChainIndices;
+		// if we haven't got to root, this is not valid chain
+		if (GetBoneChain(TransformModifier, Root, Bone, ChainIndices))
+		{
+			auto CalculateStrength = [&](int32 InBoneChainDepth, const int32 MaxDepth, float CurrentStrength, float MinStrength) -> float
+			{
+				const float Range = FMath::Max(CurrentStrength - MinStrength, 0.f);
+				const float ApplicationStrength = (float)(1.f - (float)InBoneChainDepth / (float)MaxDepth) * Range;
+				return ApplicationStrength + MinStrength;
+				//return FMath::Clamp(ApplicationStrength + MinStrength, MinStrength, CurrentStrength);
+			};
+
+			auto UpdateMotionStrength = [&](int32 InBoneChainDepth,
+				const int32 MaxPositionDepth, const int32 MaxRotationDepth, FFBIKLinkData& InOutNewLink)
+			{
+				// add motion scales
+				float LinearMotionStrength;
+				float AngularMotionStrength;
+
+				if (CurrentEffector.PositionDepth <= InBoneChainDepth)
+				{
+					LinearMotionStrength = 0.f;
+				}
+				else
+				{
+					LinearMotionStrength = CalculateStrength(InBoneChainDepth, MaxPositionDepth, SolverProperty.LinearMotionStrength, SolverProperty.MinLinearMotionStrength);
+				}
+
+				if (CurrentEffector.RotationDepth <= InBoneChainDepth)
+				{
+					AngularMotionStrength = 0.f;
+				}
+				else
+				{
+					AngularMotionStrength = CalculateStrength(InBoneChainDepth, MaxRotationDepth, SolverProperty.AngularMotionStrength, SolverProperty.MinAngularMotionStrength);
+				}
+
+				InOutNewLink.AddMotionStrength(LinearMotionStrength, AngularMotionStrength);
+			};
+
+			// position depth and rotation depth can't go beyond of it
+			// for now we cull it. 
+			const int32 PositionDepth = FMath::Min(CurrentEffector.PositionDepth, ChainIndices.Num());
+			const int32 RotationDepth = FMath::Min(CurrentEffector.RotationDepth, ChainIndices.Num());
+
+			float ChainLength = 0.f;
+
+			// add to link data
+			for (int32 BoneChainIndex = 0; BoneChainIndex < ChainIndices.Num(); ++BoneChainIndex)
+			{
+				const int32 CurrentItem = ChainIndices[BoneChainIndex];
+				ChainLength += EnsureToAddBoneToLinkData(TransformModifier, CurrentItem, LinkData,
+					HierarchyToLinkDataMap, LinkDataToHierarchyIndices);
+
+				const int32 ChainDepth = ChainIndices.Num() - BoneChainIndex;
+				int32* FoundLinkIndex = HierarchyToLinkDataMap.Find(CurrentItem);
+
+				// now we should always have it
+				check(FoundLinkIndex);
+				UpdateMotionStrength(ChainDepth, PositionDepth, RotationDepth, LinkData[*FoundLinkIndex]);
+			}
+
+			// add to EffectorTargets
+			AddToEffectorTarget(Index, BoneIndex, EffectorTargets, HierarchyToLinkDataMap, EffectorLinkIndices, ChainLength, ChainIndices, PositionDepth, RotationDepth);
+		}
+	}
+}
+/////////////////////////////////////////////////
+void UFullBodyIKSolver::InitInternal(const FIKRigTransformModifier& InGlobalTransform)
 {
 	if (UFullBodyIKSolverDefinition* SolverDef = Cast<UFullBodyIKSolverDefinition>(SolverDefinition))
 	{
-// 		LinkData.Reset();
-// 		EffectorTargets.Reset();
-// 		EffectorLinkIndices.Reset();
-// 		LinkDataToHierarchyIndices.Reset();
-// 		HierarchyToLinkDataMap.Reset();
-// 
-// 		// verify the chain
-// 		AddEffectors(Hierarchy, Root, Effectors, LinkData, EffectorTargets, EffectorLinkIndices, LinkDataToHierarchyIndices, HierarchyToLinkDataMap, SolverProperty);
+		LinkData.Reset();
+		EffectorTargets.Reset();
+		EffectorLinkIndices.Reset();
+		LinkDataToHierarchyIndices.Reset();
+		HierarchyToLinkDataMap.Reset();
+
+		// verify the chain
+		AddEffectors(InGlobalTransform, SolverDef->Root, SolverDef->Effectors, LinkData, EffectorTargets, EffectorLinkIndices, LinkDataToHierarchyIndices, HierarchyToLinkDataMap, SolverDef->SolverProperty);
 	}
 }
 
@@ -40,17 +227,20 @@ bool UFullBodyIKSolver::IsSolverActive() const
 	return false;
 }
 
-void UFullBodyIKSolver::SolveInternal(FIKRigTransformModifier& InOutGlobalTransform)
+void UFullBodyIKSolver::SolveInternal(FIKRigTransformModifier& InOutGlobalTransform, FControlRigDrawInterface* InOutDrawInterface)
 {
-#if 0 
-	if (LinkDataToHierarchyIndices.Num() > 0)
+
+	UFullBodyIKSolverDefinition* SolverDef = Cast<UFullBodyIKSolverDefinition>(SolverDefinition);
+
+	if (SolverDef && LinkDataToHierarchyIndices.Num() > 0)
 	{
+#if 0
 		// we do this every frame for now
 		if (Constraints.Num() > 0)
 		{
 			DECLARE_SCOPE_HIERARCHICAL_COUNTER(TEXT("Build Constraint"))
 			//Build constraints
-			FBIKConstraintLib::BuildConstraints(Constraints, InternalConstraints, Hierarchy, LinkData, LinkDataToHierarchyIndices, HierarchyToLinkDataMap);
+			FBIKConstraintLib::BuildConstraints(Constraints, InternalConstraints, InOutGlobalTransform, LinkData, LinkDataToHierarchyIndices, HierarchyToLinkDataMap);
 		}
 
 		// during only editor and update
@@ -64,15 +254,16 @@ void UFullBodyIKSolver::SolveInternal(FIKRigTransformModifier& InOutGlobalTransf
 		{
 			WorkData.IKSolver.ClearPostProcessDelegateForIteration();
 		}
-		// first set mid effector's transform
+	//disable constraint for now
+#endif // 0 
 
 		// before update we finalize motion scale
 		// this code may go away once we have constraint
 		// update link data and end effectors
 		for (int32 LinkIndex = 0; LinkIndex < LinkData.Num(); ++LinkIndex)
 		{
-			const FRigElementKey& Item = *LinkDataToHierarchyIndices.Find(LinkIndex);
-			LinkData[LinkIndex].SetTransform(Hierarchy->GetGlobalTransform(Item));
+			const int32 Item = *LinkDataToHierarchyIndices.Find(LinkIndex);
+			LinkData[LinkIndex].SetTransform(InOutGlobalTransform.GetGlobalTransform(Item));
 			// @todo: fix this somewhere else - we can add this to prepare step
 			// @todo: we update motion scale here, then?
 			LinkData[LinkIndex].FinalizeForSolver();
@@ -80,10 +271,12 @@ void UFullBodyIKSolver::SolveInternal(FIKRigTransformModifier& InOutGlobalTransf
 
 		// update mid effector info
 
-		const float LinearMotionStrength = FMath::Max(SolverProperty.LinearMotionStrength, SolverProperty.MinLinearMotionStrength);
-		const float AngularMotionStrength = FMath::Max(SolverProperty.AngularMotionStrength, SolverProperty.MinAngularMotionStrength);
-		const float LinearRange = LinearMotionStrength - SolverProperty.MinLinearMotionStrength;
-		const float AngularRange = AngularMotionStrength - SolverProperty.MinAngularMotionStrength;
+		const float LinearMotionStrength = FMath::Max(SolverDef->SolverProperty.LinearMotionStrength, SolverDef->SolverProperty.MinLinearMotionStrength);
+		const float AngularMotionStrength = FMath::Max(SolverDef->SolverProperty.AngularMotionStrength, SolverDef->SolverProperty.MinAngularMotionStrength);
+		const float LinearRange = LinearMotionStrength - SolverDef->SolverProperty.MinLinearMotionStrength;
+		const float AngularRange = AngularMotionStrength - SolverDef->SolverProperty.MinAngularMotionStrength;
+
+		const TArray<FFBIKRigEffector>& Effectors = SolverDef->Effectors;
 
 		// update end effector info
 		for (int32 EffectorIndex = 0; EffectorIndex < Effectors.Num(); ++EffectorIndex)
@@ -94,23 +287,23 @@ void UFullBodyIKSolver::SolveInternal(FIKRigTransformModifier& InOutGlobalTransf
 				FFBIKEffectorTarget* EffectorTarget = EffectorTargets.Find(EffectorLinkIndex);
 				if (EffectorTarget)
 				{
-					const FFBIKEndEffector& CurEffector = Effectors[EffectorIndex];
+					const FFBIKRigEffector& CurEffector = Effectors[EffectorIndex];
 					const FVector CurrentLinkLocation = LinkData[EffectorLinkIndex].GetTransform().GetLocation();
 					const FQuat CurrentLinkRotation = LinkData[EffectorLinkIndex].GetTransform().GetRotation();
-					const FVector& EffectorLocation = CurEffector.Position;
-					const FQuat& EffectorRotation = CurEffector.Rotation;
-					EffectorTarget->Position = FMath::Lerp(CurrentLinkLocation, EffectorLocation, CurEffector.PositionAlpha);
-					EffectorTarget->Rotation = FMath::Lerp(CurrentLinkRotation, EffectorRotation, CurEffector.RotationAlpha);
-					EffectorTarget->InitialPositionDistance = (EffectorLocation - CurrentLinkLocation).Size();
-					EffectorTarget->InitialRotationDistance = (FBIKUtil::GetScaledRotationAxis(EffectorRotation) - FBIKUtil::GetScaledRotationAxis(CurrentLinkRotation)).Size();
+					FIKRigTarget RigTarget;
+					ensure(GetEffectorTarget(CurEffector.Target, RigTarget));
+					EffectorTarget->Position = RigTarget.PositionTarget.Position; // FMath::Lerp(CurrentLinkLocation, EffectorLocation, RigTarget.PositionTarget.);
+					EffectorTarget->Rotation = RigTarget.RotationTarget.Rotation.Quaternion(); // FMath::Lerp(CurrentLinkRotation, EffectorRotation, CurEffector.RotationAlpha);
+					EffectorTarget->InitialPositionDistance = (EffectorTarget->Position - CurrentLinkLocation).Size();
+					EffectorTarget->InitialRotationDistance = (FBIKUtil::GetScaledRotationAxis(EffectorTarget->Rotation) - FBIKUtil::GetScaledRotationAxis(CurrentLinkRotation)).Size();
 
 					const float Pull = FMath::Clamp(CurEffector.Pull, 0.f, 1.f);
 					// we want some impact of Pull, in order for Pull to have some impact, we clamp to some number
-					const float TargetClamp = FMath::Clamp(SolverProperty.DefaultTargetClamp, 0.f, 0.7f);
+					const float TargetClamp = FMath::Clamp(SolverDef->SolverProperty.DefaultTargetClamp, 0.f, 0.7f);
 					const float Scale = TargetClamp + Pull * (1.f - TargetClamp);
 					// Pull set up
-					EffectorTarget->LinearMotionStrength = LinearRange * Scale + SolverProperty.MinLinearMotionStrength;
-					EffectorTarget->AngularMotionStrength = AngularRange * Scale + SolverProperty.MinAngularMotionStrength;
+					EffectorTarget->LinearMotionStrength = LinearRange * Scale + SolverDef->SolverProperty.MinLinearMotionStrength;
+					EffectorTarget->AngularMotionStrength = AngularRange * Scale + SolverDef->SolverProperty.MinAngularMotionStrength;
 					EffectorTarget->ConvergeScale = Scale;
 					EffectorTarget->TargetClampScale = Scale;
 
@@ -120,20 +313,17 @@ void UFullBodyIKSolver::SolveInternal(FIKRigTransformModifier& InOutGlobalTransf
 			}
 		}
 
-		TArray<FJacobianDebugData>& DebugData = WorkData.DebugData;
 		DebugData.Reset();
 
-		const bool bDebugEnabled = DebugOption.bDrawDebugHierarchy || DebugOption.bDrawDebugEffector || DebugOption.bDrawDebugConstraints;
+		const bool bDebugEnabled = SolverDef->DebugOption.bDrawDebugHierarchy || SolverDef->DebugOption.bDrawDebugEffector || SolverDef->DebugOption.bDrawDebugConstraints;
 
 		// we can't reuse memory until we fix the memory issue on RigVM
 		{
-			DECLARE_SCOPE_HIERARCHICAL_COUNTER(TEXT("Solver"))
-			FJacobianSolver_FullbodyIK& IKSolver = WorkData.IKSolver;
 			IKSolver.SolveJacobianIK(LinkData, EffectorTargets,
-				JacobianIK::FSolverParameter(SolverProperty.Damping, true, false, (SolverProperty.bUseJacobianTranspose) ? EJacobianSolver::JacobianTranspose : EJacobianSolver::JacobianPIDLS),
-				SolverProperty.MaxIterations, SolverProperty.Precision, (bDebugEnabled) ? &DebugData : nullptr);
+				JacobianIK::FSolverParameter(SolverDef->SolverProperty.Damping, true, false, (SolverDef->SolverProperty.bUseJacobianTranspose) ? EJacobianSolver::JacobianTranspose : EJacobianSolver::JacobianPIDLS),
+				SolverDef->SolverProperty.MaxIterations, SolverDef->SolverProperty.Precision, &DebugData);
 
-			if (MotionProperty.bForceEffectorRotationTarget)
+			if (SolverDef->MotionProperty.bForceEffectorRotationTarget)
 			{
 				// if position is reached, we force rotation target
 				for (int32 EffectorIndex = 0; EffectorIndex < Effectors.Num(); ++EffectorIndex)
@@ -146,13 +336,13 @@ void UFullBodyIKSolver::SolveInternal(FIKRigTransformModifier& InOutGlobalTransf
 						{
 							bool bApplyRotation = true;
 
-							if (MotionProperty.bOnlyApplyWhenReachedToTarget)
+							if (SolverDef->MotionProperty.bOnlyApplyWhenReachedToTarget)
 							{
 								// only do this when position is reached? This will conflict with converge scale
 								const FVector& BonePosition = LinkData[EffectorLinkIndex].GetTransform().GetLocation();
 								const FVector& TargetPosition = EffectorTarget->Position;
 
-								bApplyRotation = (FVector(BonePosition - TargetPosition).SizeSquared() <= SolverProperty.Precision * SolverProperty.Precision);
+								bApplyRotation = (FVector(BonePosition - TargetPosition).SizeSquared() <= SolverDef->SolverProperty.Precision * SolverDef->SolverProperty.Precision);
 							}
 
 							if (bApplyRotation)
@@ -167,11 +357,10 @@ void UFullBodyIKSolver::SolveInternal(FIKRigTransformModifier& InOutGlobalTransf
 				}
 			}
 		}
-
 		///////////////////////////////////////////////////////////////////////////
 		// debug draw start
 		///////////////////////////////////////////////////////////////////////////
-		if (bDebugEnabled && Context.DrawInterface != nullptr)
+		if (bDebugEnabled && InOutDrawInterface != nullptr)
 		{
 			const int32 DebugDataNum = DebugData.Num();
 			if (DebugData.Num() > 0)
@@ -180,10 +369,10 @@ void UFullBodyIKSolver::SolveInternal(FIKRigTransformModifier& InOutGlobalTransf
 				{
 					const TArray<FFBIKLinkData>& LocalLink = DebugData[DebugIndex].LinkData;
 
-					FTransform Offset = DebugOption.DrawWorldOffset;
+					FTransform Offset = SolverDef->DebugOption.DrawWorldOffset;
 					Offset.SetLocation(Offset.GetLocation() * (DebugDataNum - DebugIndex));
 
-					if (DebugOption.bDrawDebugHierarchy)
+					if (SolverDef->DebugOption.bDrawDebugHierarchy)
 					{
 						for (int32 LinkIndex = 0; LinkIndex < LocalLink.Num(); ++LinkIndex)
 						{
@@ -192,25 +381,25 @@ void UFullBodyIKSolver::SolveInternal(FIKRigTransformModifier& InOutGlobalTransf
 							FLinearColor DrawColor = FLinearColor::White;
 
 							float LineThickness = 0.f;
-							if (DebugOption.bColorAngularMotionStrength || DebugOption.bColorLinearMotionStrength)
+							if (SolverDef->DebugOption.bColorAngularMotionStrength || SolverDef->DebugOption.bColorLinearMotionStrength)
 							{
 								DrawColor = FLinearColor::Black;
-								if (DebugOption.bColorAngularMotionStrength)
+								if (SolverDef->DebugOption.bColorAngularMotionStrength)
 								{
-									const float Range = FMath::Max(SolverProperty.AngularMotionStrength - SolverProperty.MinAngularMotionStrength, 0.f);
+									const float Range = FMath::Max(SolverDef->SolverProperty.AngularMotionStrength - SolverDef->SolverProperty.MinAngularMotionStrength, 0.f);
 									if (Range > 0.f)
 									{
-										float CurrentStrength = Data.GetAngularMotionStrength() - SolverProperty.MinAngularMotionStrength;
+										float CurrentStrength = Data.GetAngularMotionStrength() - SolverDef->SolverProperty.MinAngularMotionStrength;
 										float Alpha = FMath::Clamp(CurrentStrength / Range, 0.f, 1.f);
 										DrawColor.R = LineThickness = Alpha;
 									}
 								}
-								else if (DebugOption.bColorLinearMotionStrength)
+								else if (SolverDef->DebugOption.bColorLinearMotionStrength)
 								{
-									const float Range = FMath::Max(SolverProperty.LinearMotionStrength - SolverProperty.MinLinearMotionStrength, 0.f);
+									const float Range = FMath::Max(SolverDef->SolverProperty.LinearMotionStrength - SolverDef->SolverProperty.MinLinearMotionStrength, 0.f);
 									if (Range > 0.f)
 									{
-										float CurrentStrength = Data.GetLinearMotionStrength() - SolverProperty.MinLinearMotionStrength;
+										float CurrentStrength = Data.GetLinearMotionStrength() - SolverDef->SolverProperty.MinLinearMotionStrength;
 										float Alpha = FMath::Clamp(CurrentStrength / Range, 0.f, 1.f);
 										DrawColor.B = LineThickness = Alpha;
 									}
@@ -220,17 +409,17 @@ void UFullBodyIKSolver::SolveInternal(FIKRigTransformModifier& InOutGlobalTransf
 							if (Data.ParentLinkIndex != INDEX_NONE)
 							{
 								const FFBIKLinkData& ParentData = LocalLink[Data.ParentLinkIndex];
-								Context.DrawInterface->DrawLine(Offset, Data.GetPreviousTransform().GetLocation(), ParentData.GetPreviousTransform().GetLocation(), DrawColor, LineThickness);
+								InOutDrawInterface->DrawLine(Offset, Data.GetPreviousTransform().GetLocation(), ParentData.GetPreviousTransform().GetLocation(), DrawColor, LineThickness);
 							}
 
-							if (DebugOption.bDrawDebugAxes)
+							if (SolverDef->DebugOption.bDrawDebugAxes)
 							{
-								Context.DrawInterface->DrawAxes(Offset, Data.GetPreviousTransform(), DebugOption.DrawSize);
+								InOutDrawInterface->DrawAxes(Offset, Data.GetPreviousTransform(), SolverDef->DebugOption.DrawSize);
 							}
 						}
 					}
 
-					if (DebugOption.bDrawDebugEffector)
+					if (SolverDef->DebugOption.bDrawDebugEffector)
 					{
 						for (auto Iter = EffectorTargets.CreateConstIterator(); Iter; ++Iter)
 						{
@@ -238,24 +427,24 @@ void UFullBodyIKSolver::SolveInternal(FIKRigTransformModifier& InOutGlobalTransf
 							if (EffectorTarget.bPositionEnabled)
 							{
 								// draw effector target locations
-								Context.DrawInterface->DrawBox(Offset, FTransform(EffectorTarget.Position), FLinearColor::Yellow, DebugOption.DrawSize);
+								InOutDrawInterface->DrawBox(Offset, FTransform(EffectorTarget.Position), FLinearColor::Yellow, SolverDef->DebugOption.DrawSize);
 							}
 
 							// draw effector link location
-							Context.DrawInterface->DrawBox(Offset, LocalLink[Iter.Key()].GetPreviousTransform(), FLinearColor::Green, DebugOption.DrawSize);
+							InOutDrawInterface->DrawBox(Offset, LocalLink[Iter.Key()].GetPreviousTransform(), FLinearColor::Green, SolverDef->DebugOption.DrawSize);
 						}
 
 						for (int32 Index = 0; Index < DebugData[DebugIndex].TargetVectorSources.Num(); ++Index)
 						{
 							// draw arrow to the target
-							Context.DrawInterface->DrawLine(Offset, DebugData[DebugIndex].TargetVectorSources[Index].GetLocation(),
+							InOutDrawInterface->DrawLine(Offset, DebugData[DebugIndex].TargetVectorSources[Index].GetLocation(),
 								DebugData[DebugIndex].TargetVectorSources[Index].GetLocation() + DebugData[DebugIndex].TargetVectors[Index], FLinearColor::Red);
 						}
 					}
 				}
 			}
-
-			if (DebugOption.bDrawDebugConstraints && InternalConstraints.Num())
+#if 0 
+			if (SolverDef->DebugOption.bDrawDebugConstraints && InternalConstraints.Num())
 			{
 				FTransform Offset = FTransform::Identity;
 
@@ -271,7 +460,7 @@ void UFullBodyIKSolver::SolveInternal(FIKRigTransformModifier& InOutGlobalTransf
 							{
 								FTransform ConstraintFrame = LinkData[*Found].GetTransform();
 								ConstraintFrame.ConcatenateRotation(FQuat(Constraints[Index].OffsetRotation));
-								Context.DrawInterface->DrawAxes(Offset, ConstraintFrame, 2.f);
+								InOutDrawInterface->DrawAxes(Offset, ConstraintFrame, 2.f);
 							}
 						}
 					}
@@ -288,7 +477,7 @@ void UFullBodyIKSolver::SolveInternal(FIKRigTransformModifier& InOutGlobalTransf
 						// base is parent transform but in their space, we can get there by inversing local ref rotation
 						FTransform BaseTransform = FTransform(LocalRefRotation).GetRelativeTransformReverse(RotationTransform);
 						BaseTransform.ConcatenateRotation(LimitConstraint.BaseFrameOffset);
-						Context.DrawInterface->DrawAxes(Offset, BaseTransform, 5.f, 1.f);
+						InOutDrawInterface->DrawAxes(Offset, BaseTransform, 5.f, 1.f);
 
 						// current transform
 						const FQuat LocalRotation = BaseTransform.GetRotation().Inverse() * RotationTransform.GetRotation();
@@ -298,7 +487,7 @@ void UFullBodyIKSolver::SolveInternal(FIKRigTransformModifier& InOutGlobalTransf
 						RotationTransform.SetLocation(BaseTransform.GetLocation());
 
 						// draw ref pose on their current transform
-						Context.DrawInterface->DrawAxes(Offset, RotationTransform, 10.f, 1.f);
+						InOutDrawInterface->DrawAxes(Offset, RotationTransform, 10.f, 1.f);
 
 						FVector XAxis = BaseTransform.GetUnitAxis(EAxis::X);
 						FVector YAxis = BaseTransform.GetUnitAxis(EAxis::Y);
@@ -309,7 +498,7 @@ void UFullBodyIKSolver::SolveInternal(FIKRigTransformModifier& InOutGlobalTransf
 							FTransform XAxisConeTM(YAxis, XAxis ^ YAxis, XAxis, BaseTransform.GetTranslation());
 							XAxisConeTM.SetRotation(FQuat(XAxis, 0.f) * XAxisConeTM.GetRotation());
 							XAxisConeTM.SetScale3D(FVector(30.f));
-							Context.DrawInterface->DrawCone(Offset, XAxisConeTM, LimitConstraint.Limit.X, 0.0f, 24, false, FLinearColor::Red, GEngine->ConstraintLimitMaterialX->GetRenderProxy());
+							InOutDrawInterface->DrawCone(Offset, XAxisConeTM, LimitConstraint.Limit.X, 0.0f, 24, false, FLinearColor::Red, GEngine->ConstraintLimitMaterialX->GetRenderProxy());
 						}
 
 						if (LimitConstraint.bYLimitSet)
@@ -317,7 +506,7 @@ void UFullBodyIKSolver::SolveInternal(FIKRigTransformModifier& InOutGlobalTransf
 							FTransform YAxisConeTM(ZAxis, YAxis ^ ZAxis, YAxis, BaseTransform.GetTranslation());
 							YAxisConeTM.SetRotation(FQuat(YAxis, 0.f) * YAxisConeTM.GetRotation());
 							YAxisConeTM.SetScale3D(FVector(30.f));
-							Context.DrawInterface->DrawCone(Offset, YAxisConeTM, LimitConstraint.Limit.Y, 0.0f, 24, false, FLinearColor::Green, GEngine->ConstraintLimitMaterialY->GetRenderProxy());
+							InOutDrawInterface->DrawCone(Offset, YAxisConeTM, LimitConstraint.Limit.Y, 0.0f, 24, false, FLinearColor::Green, GEngine->ConstraintLimitMaterialY->GetRenderProxy());
 						}
 
 						if (LimitConstraint.bZLimitSet)
@@ -325,7 +514,7 @@ void UFullBodyIKSolver::SolveInternal(FIKRigTransformModifier& InOutGlobalTransf
 							FTransform ZAxisConeTM(XAxis, ZAxis ^ XAxis, ZAxis, BaseTransform.GetTranslation());
 							ZAxisConeTM.SetRotation(FQuat(ZAxis, 0.f) * ZAxisConeTM.GetRotation());
 							ZAxisConeTM.SetScale3D(FVector(30.f));
-							Context.DrawInterface->DrawCone(Offset, ZAxisConeTM, LimitConstraint.Limit.Z, 0.0f, 24, false, FLinearColor::Blue, GEngine->ConstraintLimitMaterialZ->GetRenderProxy());
+							InOutDrawInterface->DrawCone(Offset, ZAxisConeTM, LimitConstraint.Limit.Z, 0.0f, 24, false, FLinearColor::Blue, GEngine->ConstraintLimitMaterialZ->GetRenderProxy());
 						}
 					}
 
@@ -349,11 +538,13 @@ void UFullBodyIKSolver::SolveInternal(FIKRigTransformModifier& InOutGlobalTransf
 						Positions.Add(JointTarget);
 						Positions.Add(RootTransform.GetLocation());
 
-						Context.DrawInterface->DrawLines(Offset, Positions, FLinearColor::Gray, 1.2f);
-						Context.DrawInterface->DrawLine(Offset, JointTransform.GetLocation(), JointTarget, FLinearColor::Red, 1.2f);
+						InOutDrawInterface->DrawLines(Offset, Positions, FLinearColor::Gray, 1.2f);
+						InOutDrawInterface->DrawLine(Offset, JointTransform.GetLocation(), JointTarget, FLinearColor::Red, 1.2f);
 					}
 				}
 			}
+
+			#endif // 0 - disable all constraint
 		}
 		///////////////////////////////////////////////////////////////////////////
 		// debug draw end
@@ -365,10 +556,8 @@ void UFullBodyIKSolver::SolveInternal(FIKRigTransformModifier& InOutGlobalTransf
 	{
 		// only propagate, if you are leaf joints here
 		// this means, only the last joint in the test
-		const FRigElementKey& CurrentItem = *LinkDataToHierarchyIndices.Find(LinkIndex);
+		const int32 CurrentItem = *LinkDataToHierarchyIndices.Find(LinkIndex);
 		const FTransform& LinkTransform = LinkData[LinkIndex].GetTransform();
-		Hierarchy->SetGlobalTransform(CurrentItem, LinkTransform, bPropagateToChildren);
+		InOutGlobalTransform.SetGlobalTransform(CurrentItem, LinkTransform, true);
 	}
-
-	#endif // 0 
 }
\ No newline at end of file
diff --git a/Engine/Plugins/Experimental/FullBodyIK/Source/FullBodyIK/Private/FullBodyIKSolver.h b/Engine/Plugins/Experimental/FullBodyIK/Source/FullBodyIK/Private/FullBodyIKSolver.h
--- a/Engine/Plugins/Experimental/FullBodyIK/Source/FullBodyIK/Private/FullBodyIKSolver.h
+++ b/Engine/Plugins/Experimental/FullBodyIK/Source/FullBodyIK/Private/FullBodyIKSolver.h
@@ -25,10 +25,10 @@ class FULLBODYIK_API UFullBodyIKSolver : public UIKRigSolver
 	TMap<int32, FFBIKEffectorTarget> EffectorTargets; 
 	/** End Effector Link Indices - EndEffector index to LinkData index*/
 	TArray<int32> EffectorLinkIndices;
-	/** Map from LinkData index to Rig Hierarchy Index*/
-	TMap<int32, FRigElementKey> LinkDataToHierarchyIndices;
+	/** Map from LinkData index to Hierarchy Index*/
+	TMap<int32, int32> LinkDataToHierarchyIndices;
 	/** Map from Rig Hierarchy Index to LinkData index*/
-	TMap<FRigElementKey, int32> HierarchyToLinkDataMap;
+	TMap<int32, int32> HierarchyToLinkDataMap;
 	/** Constraints data */
 	TArray<ConstraintType> InternalConstraints;
 	/* Current Solver */
@@ -37,8 +37,8 @@ class FULLBODYIK_API UFullBodyIKSolver : public UIKRigSolver
 	TArray<FJacobianDebugData> DebugData;
 
 protected:
-	virtual void InitInternal() override;
-	virtual void SolveInternal(FIKRigTransformModifier& InOutGlobalTransform) override;
+	virtual void InitInternal(const FIKRigTransformModifier& InGlobalTransform) override;
+	virtual void SolveInternal(FIKRigTransformModifier& InOutGlobalTransform, FControlRigDrawInterface* InOutDrawInterface) override;
 	virtual bool IsSolverActive() const override;
 };
 
diff --git a/Engine/Plugins/Experimental/FullBodyIK/Source/FullBodyIK/Private/FullBodyIKSolverDefinition.h b/Engine/Plugins/Experimental/FullBodyIK/Source/FullBodyIK/Private/FullBodyIKSolverDefinition.h
--- a/Engine/Plugins/Experimental/FullBodyIK/Source/FullBodyIK/Private/FullBodyIKSolverDefinition.h
+++ b/Engine/Plugins/Experimental/FullBodyIK/Source/FullBodyIK/Private/FullBodyIKSolverDefinition.h
@@ -10,6 +10,7 @@
 #include "IKRigSolverDefinition.h"
 #include "FBIKShared.h"
 #include "FBIKConstraintOption.h"
+#include "FBIKDebugOption.h"
 #include "FullBodyIKSolverDefinition.generated.h"
 
 USTRUCT()
@@ -62,6 +63,8 @@ public:
 	UPROPERTY(EditAnywhere, Category = Motion)
 	FMotionProcessInput MotionProperty;
 
+	UPROPERTY(EditAnywhere, Category = Solver)
+	FFBIKDebugOption DebugOption;
 private:
 #if WITH_EDITOR
 	virtual void UpdateEffectors() override;
diff --git a/Engine/Plugins/Experimental/FullBodyIK/Source/FullBodyIK/Public/FBIKDebugOption.h b/Engine/Plugins/Experimental/FullBodyIK/Source/FullBodyIK/Public/FBIKDebugOption.h
--- a/Engine/Plugins/Experimental/FullBodyIK/Source/FullBodyIK/Public/FBIKDebugOption.h
+++ b/Engine/Plugins/Experimental/FullBodyIK/Source/FullBodyIK/Public/FBIKDebugOption.h
@@ -10,30 +10,30 @@ struct FFBIKDebugOption
 {
 	GENERATED_BODY()
 	
-	UPROPERTY(meta = (Input))
+	UPROPERTY(EditAnywhere, Category = FFBIKDebugOption)
 	bool bDrawDebugHierarchy = false;
 
 	// use red channel
-	UPROPERTY(meta = (Input))
+	UPROPERTY(EditAnywhere, Category = FFBIKDebugOption)
 	bool bColorAngularMotionStrength = false;
 
 	// use green channel
-	UPROPERTY(meta = (Input))
+	UPROPERTY(EditAnywhere, Category = FFBIKDebugOption)
 	bool bColorLinearMotionStrength = false;
 
-	UPROPERTY(meta = (Input))
+	UPROPERTY(EditAnywhere, Category = FFBIKDebugOption)
 	bool bDrawDebugAxes = false;
 
-	UPROPERTY(meta = (Input))
+	UPROPERTY(EditAnywhere, Category = FFBIKDebugOption)
 	bool bDrawDebugEffector = false;
 
-	UPROPERTY(meta = (Input))
+	UPROPERTY(EditAnywhere, Category = FFBIKDebugOption)
 	bool bDrawDebugConstraints = false;
 
-	UPROPERTY(meta = (Input))
+	UPROPERTY(EditAnywhere, Category = FFBIKDebugOption)
 	FTransform DrawWorldOffset;
 
-	UPROPERTY(meta = (Input))
+	UPROPERTY(EditAnywhere, Category = FFBIKDebugOption)
 	float DrawSize = 5.f;
 
 	FFBIKDebugOption()
