Merging IKRigSolver and IKRigSolverDefinition classes

diff --git a/Engine/Plugins/Animation/IKRig/Source/IKRig/Private/Constraints/RotationLimitConstraint.cpp b/Engine/Plugins/Animation/IKRig/Source/IKRig/Private/Constraints/RotationLimitConstraint.cpp
--- a/Engine/Plugins/Animation/IKRig/Source/IKRig/Private/Constraints/RotationLimitConstraint.cpp
+++ b/Engine/Plugins/Animation/IKRig/Source/IKRig/Private/Constraints/RotationLimitConstraint.cpp
@@ -1,10 +1,5 @@
 // Copyright Epic Games, Inc. All Rights Reserved.
 
-/**
- * Contains IKRigConstraint Runtime Implementation
- *
- */
-
 #include "RotationLimitConstraint.h"
 #include "IKRigHierarchy.h"
 #include "IKRigDataTypes.h"
diff --git a/Engine/Plugins/Animation/IKRig/Source/IKRig/Private/IKRigConstraintDefinition.cpp b/Engine/Plugins/Animation/IKRig/Source/IKRig/Private/IKRigConstraintDefinition.cpp
deleted file mode 100644
diff --git a/Engine/Plugins/Animation/IKRig/Source/IKRig/Private/IKRigDefinition.cpp b/Engine/Plugins/Animation/IKRig/Source/IKRig/Private/IKRigDefinition.cpp
--- a/Engine/Plugins/Animation/IKRig/Source/IKRig/Private/IKRigDefinition.cpp
+++ b/Engine/Plugins/Animation/IKRig/Source/IKRig/Private/IKRigDefinition.cpp
@@ -1,17 +1,10 @@
 // Copyright Epic Games, Inc. All Rights Reserved.
 
-/*=============================================================================
-	IKRigDefinition.cpp: Composite classes that contains sequence for each section
-=============================================================================*/
-
 #include "IKRigDefinition.h"
-#include "IKRigSolverDefinition.h"
-#include "IKRigConstraintDefinition.h"
+#include "IKRigSolver.h"
 
 UIKRigDefinition::UIKRigDefinition()
 {
-	// we create constraint definition for them
-	ConstraintDefinitions = CreateDefaultSubobject<UIKRigConstraintDefinition>(TEXT("ConstraintDefinition"));
 }
 
 #if WITH_EDITOR
@@ -23,11 +16,11 @@ void UIKRigDefinition::UpdateGoal()
 	Sanitize();
 
 	TArray<FName> ListOfGoals;
-	for (UIKRigSolverDefinition* SolverDef : SolverDefinitions)
+	for (UIKRigSolver* Solver : Solvers)
 	{
-		if (SolverDef)
+		if (Solver)
 		{
-			SolverDef->CollectGoals(ListOfGoals);
+			Solver->CollectGoals(ListOfGoals);
 		}
 	}
 
@@ -72,7 +65,7 @@ void UIKRigDefinition::UpdateGoal()
 void UIKRigDefinition::Sanitize()
 {
 	// sanitize to clean up
-	SolverDefinitions.RemoveAll([this](const UIKRigSolverDefinition* IKDef) { return IKDef == nullptr; });
+	Solvers.RemoveAll([this](const UIKRigSolver* Solver) { return Solver == nullptr; });
 }
 
 /////////////////////////////////////////////////////////
diff --git a/Engine/Plugins/Animation/IKRig/Source/IKRig/Private/IKRigProcessor.cpp b/Engine/Plugins/Animation/IKRig/Source/IKRig/Private/IKRigProcessor.cpp
--- a/Engine/Plugins/Animation/IKRig/Source/IKRig/Private/IKRigProcessor.cpp
+++ b/Engine/Plugins/Animation/IKRig/Source/IKRig/Private/IKRigProcessor.cpp
@@ -7,7 +7,6 @@
 
 #include "IKRigProcessor.h"
 #include "IKRigDefinition.h"
-#include "IKRigSolverDefinition.h"
 #include "IKRigSolver.h"
 #include "Solvers/ConstraintSolver.h"
 
@@ -40,40 +39,25 @@ void UIKRigProcessor::Initialize(const FIKRigTransform& InRefTransform)
 void UIKRigProcessor::Reinitialize()
 {
 	TransformModifier.ResetGlobalTransform(ReferenceTransform);
-
 	IKGoals = RigDefinition->GetGoals();
-
-	// i'm copying this here 
-	TArray<UIKRigSolverDefinition*> SolverDefinitions = RigDefinition->GetSolverDefinitions();
-
 	UIKRigSolver::FIKRigTransformGetter RefTransformGetter = UIKRigSolver::FIKRigTransformGetter::CreateUObject(this, &UIKRigProcessor::GetRefPoseGetter);
 	UIKRigSolver::FIKRigGoalGetter GoalGetter = UIKRigSolver::FIKRigGoalGetter::CreateUObject(this, &UIKRigProcessor::GoalGetter);
 
-	const int32 NumSolvers = SolverDefinitions.Num();
+	TArray<UIKRigSolver*> RigSolvers = RigDefinition->GetSolvers();
+	const int32 NumSolvers = RigSolvers.Num();
 	Solvers.Reset(NumSolvers);
 
 	for (int32 Index = 0; Index < NumSolvers; ++Index)
 	{
-		// it's possible SolverDefinitions.Num() != Solvers.Num()
-		if (SolverDefinitions[Index])
+		TSubclassOf<UIKRigSolver> ClassType = RigSolvers[Index]->GetClass();
+		if (ClassType != nullptr)
 		{
-			TSubclassOf<UIKRigSolver> ClassType = SolverDefinitions[Index]->GetExecutionClass();
-			if (ClassType != nullptr)
-			{
-				UIKRigSolver* NewSolver = NewObject<UIKRigSolver>(this, ClassType);
-				// const transform modifier for hierarchy query
-				NewSolver->Init(SolverDefinitions[Index], TransformModifier, RefTransformGetter, GoalGetter);
-				Solvers.Add(NewSolver);
-			}
+			UIKRigSolver* NewSolver = DuplicateObject(RigSolvers[Index], this);
+			NewSolver->Init(TransformModifier, RefTransformGetter, GoalGetter);
+			Solvers.Add(NewSolver);
 		}
 	}
 
-	// initialize constraint solver
-	CostraintSolver = NewObject<UIKRigConstraintSolver>(this);
-	check (CostraintSolver);
-	CostraintSolver->SetConstraintDefinition(RigDefinition->ConstraintDefinitions);
-	CostraintSolver->Init(RigDefinition->ConstraintDefinitions, TransformModifier, RefTransformGetter, GoalGetter);
-
 	bInitialized = true;
 }
 
@@ -112,9 +96,6 @@ void UIKRigProcessor::Solve()
 			// Draw interface is pointer in case we want to nullify for optimization
 			Solvers[Index]->Solve(TransformModifier, &DrawInterface);
 		}
-
-		// solve constraint
-		CostraintSolver->Solve(TransformModifier, &DrawInterface);
 	}
 }
 
diff --git a/Engine/Plugins/Animation/IKRig/Source/IKRig/Private/IKRigSolver.cpp b/Engine/Plugins/Animation/IKRig/Source/IKRig/Private/IKRigSolver.cpp
--- a/Engine/Plugins/Animation/IKRig/Source/IKRig/Private/IKRigSolver.cpp
+++ b/Engine/Plugins/Animation/IKRig/Source/IKRig/Private/IKRigSolver.cpp
@@ -1,17 +1,14 @@
 // Copyright Epic Games, Inc. All Rights Reserved.
 
-/*=============================================================================
-	IKRigSolver.cpp: Solver definition class
-=============================================================================*/
 
 #include "IKRigSolver.h"
-#include "IKRigSolverDefinition.h"
 #include "IKRigDataTypes.h"
+#include "IKRigDefinition.h"
 
 // input hierarchy and ref pose? 
-void UIKRigSolver::Init(UIKRigSolverDefinition* InSolverDefinition, const FIKRigTransformModifier& TransformModifier, FIKRigTransformGetter InRefPoseGetter, FIKRigGoalGetter InGoalGetter)
+void UIKRigSolver::Init(const FIKRigTransformModifier& TransformModifier, FIKRigTransformGetter InRefPoseGetter, FIKRigGoalGetter InGoalGetter)
 {
-	SolverDefinition = InSolverDefinition;
+	//SolverDefinition = InSolverDefinition;
 
 	RefPoseGetter = InRefPoseGetter;
 	ensure(RefPoseGetter.IsBound());
@@ -24,7 +21,8 @@ void UIKRigSolver::Init(UIKRigSolverDefinition* InSolverDefinition, const FIKRig
 
 bool UIKRigSolver::IsSolverActive() const 
 {
-	return (SolverDefinition && bEnabled);
+	//return (SolverDefinition && bEnabled);
+	return bEnabled;
 }
 
 // input : goal getter or goals
@@ -39,9 +37,10 @@ void UIKRigSolver::Solve(FIKRigTransformModifier& InOutGlobalTransform, FControl
 
 bool UIKRigSolver::GetEffectorTarget(const FIKRigEffector& InEffector, FIKRigTarget& OutTarget) const
 {
-	if (SolverDefinition && GoalGetter.IsBound())
+	//if (SolverDefinition && GoalGetter.IsBound())
+	if (GoalGetter.IsBound())
 	{	
-		const FName* GoalName = SolverDefinition->GetEffectorToGoal().Find(InEffector);
+		const FName* GoalName = GetEffectorToGoal().Find(InEffector);
 		if (GoalName)
 		{
 			return GoalGetter.Execute(*GoalName, OutTarget);
@@ -60,4 +59,87 @@ const FIKRigTransform& UIKRigSolver::GetReferencePose() const
 
 	static FIKRigTransform Dummy;
 	return Dummy;
-}
\ No newline at end of file
+}
+
+void UIKRigSolver::CollectGoals(TArray<FName>& OutGoals)
+{
+	TArray<FName> LocalGoals;
+	// this just accumulate on the current array
+	// if you want to clear, clear before coming here
+	EffectorToGoal.GenerateValueArray(LocalGoals);
+
+	OutGoals += LocalGoals;
+}
+
+#if WITH_EDITOR
+void UIKRigSolver::RenameGoal(const FName& OldName, const FName& NewName)
+{
+	for (auto Iter = EffectorToGoal.CreateIterator(); Iter; ++Iter)
+	{
+		if (Iter.Value() == OldName)
+		{
+			Iter.Value() = NewName;
+		}
+	}
+}
+
+void UIKRigSolver::EnsureUniqueGoalName(FName& InOutUniqueGoalName) const
+{
+	// call delegate 
+	UIKRigDefinition* IKRigDef = CastChecked<UIKRigDefinition>(GetOuter());
+	IKRigDef->EnsureCreateUniqueGoalName(InOutUniqueGoalName);
+}
+
+FName UIKRigSolver::CreateUniqueGoalName(const TCHAR* Suffix) const
+{
+	if (Suffix)
+	{
+		FString NewGoalStr = FString("") + TEXT("NewGoal_") + Suffix;
+		// replace any whitespace with _
+		NewGoalStr = NewGoalStr.Replace(TEXT(" \t"), TEXT("_"));
+		FName NewGoalName = FName(*NewGoalStr, FNAME_Add);
+		EnsureUniqueGoalName(NewGoalName);
+
+		return NewGoalName;
+	}
+
+	return NAME_None;
+}
+
+void UIKRigSolver::OnGoalHasBeenUpdated()
+{
+	GoalNeedsUpdateDelegate.Broadcast();
+}
+
+void UIKRigSolver::EnsureToAddEffector(const FIKRigEffector& InEffector, const FString& InPrefix)
+{
+	FName* GoalName = EffectorToGoal.Find(InEffector);
+	if (!GoalName)
+	{
+		EffectorToGoal.Add(InEffector) = CreateUniqueGoalName(*InPrefix);
+	}
+}
+
+void UIKRigSolver::EnsureToRemoveEffector(const FIKRigEffector& InEffector)
+{
+	EffectorToGoal.Remove(InEffector);
+}
+
+#endif // WITH_EDITOR
+
+void UIKRigSolver::PostLoad()
+{
+	Super::PostLoad();
+
+#if WITH_EDITOR
+	UpdateEffectors();
+#endif // WITH_EDITOR
+}
+
+void UIKRigSolver::Serialize(FArchive& Ar)
+{
+	Super::Serialize(Ar);
+
+	// serialize these manually because that's the custom types
+	Ar << EffectorToGoal;
+}
diff --git a/Engine/Plugins/Animation/IKRig/Source/IKRig/Private/IKRigSolverDefinition.cpp b/Engine/Plugins/Animation/IKRig/Source/IKRig/Private/IKRigSolverDefinition.cpp
deleted file mode 100644
diff --git a/Engine/Plugins/Animation/IKRig/Source/IKRig/Private/Solvers/ConstraintSolver.cpp b/Engine/Plugins/Animation/IKRig/Source/IKRig/Private/Solvers/ConstraintSolver.cpp
--- a/Engine/Plugins/Animation/IKRig/Source/IKRig/Private/Solvers/ConstraintSolver.cpp
+++ b/Engine/Plugins/Animation/IKRig/Source/IKRig/Private/Solvers/ConstraintSolver.cpp
@@ -1,12 +1,9 @@
 // Copyright Epic Games, Inc. All Rights Reserved.
 
-/**
- * Contains IKRig Constraint Solver Definition 
- *
- */
 #include "Solvers/ConstraintSolver.h"
 #include "IKRigConstraint.h"
 
+/*
 // register/unregister query function
 void UIKRigConstraintSolver::RegisterQueryConstraintHandler(const FIKRigQueryConstraint& InQueryConstraintHandler)
 {
@@ -17,9 +14,11 @@ void UIKRigConstraintSolver::UnregisterQueryConstraintHandler()
 {
 
 }
+*/
 
 void UIKRigConstraintSolver::InitInternal(const FIKRigTransformModifier& InGlobalTransform)
 {
+	/*
 	if (ConstraintDefinition)
 	{
 		// we duplicate and create a copy of the objects, so that we can mutate
@@ -39,11 +38,12 @@ void UIKRigConstraintSolver::InitInternal(const FIKRigTransformModifier& InGloba
 
 		// go through all UIKRigConstraint and create copy
 		ActiveProfile = UIKRigConstraintDefinition::DefaultProfileName;
-	}
+	}*/
 }
 
 void UIKRigConstraintSolver::SolveInternal(FIKRigTransformModifier& InOutGlobalTransform, FControlRigDrawInterface* InOutDrawInterface)
 {
+	/*
 	FIKRigConstraintProfile* Current = ConstraintProfiles.Find(ActiveProfile);
 
 	if (Current)
@@ -57,9 +57,10 @@ void UIKRigConstraintSolver::SolveInternal(FIKRigTransformModifier& InOutGlobalT
 				Constraint->Apply(InOutGlobalTransform, InOutDrawInterface);
 			}
 		}
-	}
+	}*/
 }
 
+/*
 void UIKRigConstraintSolver::SolverConstraints(const TArray<FName>& ConstraintsList, FIKRigTransformModifier& InOutGlobalTransform, FControlRigDrawInterface* InOutDrawInterface)
 {
 	FIKRigConstraintProfile* Current = ConstraintProfiles.Find(ActiveProfile);
@@ -77,18 +78,4 @@ void UIKRigConstraintSolver::SolverConstraints(const TArray<FName>& ConstraintsL
 		}
 	}
 }
-
-bool UIKRigConstraintSolver::IsSolverActive() const
-{	
-// later we can do more elaborated version
-	return Super::IsSolverActive();
-}
-
-void UIKRigConstraintSolver::SetConstraintDefinition(UIKRigConstraintDefinition* InConstraintDefinition)
-{
-	ConstraintDefinition = InConstraintDefinition;
-	// reset instanced data
-
-	ConstraintProfiles.Reset();
-	ActiveProfile = UIKRigConstraintDefinition::DefaultProfileName;
-}
\ No newline at end of file
+*/
\ No newline at end of file
diff --git a/Engine/Plugins/Animation/IKRig/Source/IKRig/Private/Solvers/TransformSolver.cpp b/Engine/Plugins/Animation/IKRig/Source/IKRig/Private/Solvers/TransformSolver.cpp
--- a/Engine/Plugins/Animation/IKRig/Source/IKRig/Private/Solvers/TransformSolver.cpp
+++ b/Engine/Plugins/Animation/IKRig/Source/IKRig/Private/Solvers/TransformSolver.cpp
@@ -1,13 +1,15 @@
 // Copyright Epic Games, Inc. All Rights Reserved.
 
-/*=============================================================================
-	TransformSolver.cpp: Solver execution class for Transform
-=============================================================================*/
 
 #include "Solvers/TransformSolver.h"
 #include "IKRigDataTypes.h"
 #include "IKRigHierarchy.h"
-#include "Solvers/TransformSolverDefinition.h"
+
+
+UTransformSolver::UTransformSolver()
+	: TransformTargetName(TEXT("TransformTarget"))
+{
+}
 
 void UTransformSolver::InitInternal(const FIKRigTransformModifier& InGlobalTransform)
 {
@@ -16,40 +18,71 @@ void UTransformSolver::InitInternal(const FIKRigTransformModifier& InGlobalTrans
 
 bool UTransformSolver::IsSolverActive() const
 {
-	if (Super::IsSolverActive())
+	return Super::IsSolverActive() && (bEnablePosition || bEnableRotation);
+}
+
+void UTransformSolver::SolveInternal(FIKRigTransformModifier& InOutGlobalTransform, FControlRigDrawInterface* InOutDrawInterface)
+{
+	FIKRigTarget Target;
+	if (!GetEffectorTarget(TransformTarget, Target))
+	{
+		return;
+	}
+	
+	int32 Index = InOutGlobalTransform.Hierarchy->GetIndex(TransformTarget.Bone);
+	if (Index == INDEX_NONE)
+	{
+		return;
+	}
+
+	FTransform CurrentTransform = InOutGlobalTransform.GetGlobalTransform(Index);
+
+	if (bEnablePosition)
+	{
+		CurrentTransform.SetLocation(Target.PositionTarget.Position);
+	}
+	if (bEnableRotation)
+	{
+		CurrentTransform.SetRotation(Target.RotationTarget.Rotation.Quaternion());
+	}
+
+	InOutGlobalTransform.SetGlobalTransform(Index, CurrentTransform, true);
+}
+
+#if WITH_EDITOR
+
+void UTransformSolver::UpdateEffectors()
+{
+	// KIARAN - this doesn't make any sense, why would we remove an effector just because
+	// the solver settings have been toggled off?
+
+	// update tasks - you want to keep old name though if changed
+	bool bActiveTask = bEnablePosition || bEnableRotation;
+
+	// not found, we add new
+	if (bActiveTask)
+	{
+		EnsureToAddEffector(TransformTarget, *TransformTargetName);
+	}
+	// if the opposite, we have to remove
+	else if (!bActiveTask)
 	{
-		if (UTransformSolverDefinition* SolverDef = Cast<UTransformSolverDefinition>(SolverDefinition))
-		{
-			return (SolverDef->bEnablePosition || SolverDef->bEnableRotation);
-		}
+		EnsureToRemoveEffector(TransformTarget);
 	}
 
-	return false;
+	// trigger a delegate for goal has been updated?
+	OnGoalHasBeenUpdated();
 }
 
-void UTransformSolver::SolveInternal(FIKRigTransformModifier& InOutGlobalTransform, FControlRigDrawInterface* InOutDrawInterface)
+void UTransformSolver::PostEditChangeChainProperty(struct FPropertyChangedChainEvent& PropertyChangedEvent)
 {
-	if (UTransformSolverDefinition* SolverDef = Cast<UTransformSolverDefinition>(SolverDefinition))
+	if (PropertyChangedEvent.Property->GetFName() == GET_MEMBER_NAME_CHECKED(UTransformSolver, bEnablePosition) ||
+		PropertyChangedEvent.Property->GetFName() == GET_MEMBER_NAME_CHECKED(UTransformSolver, bEnableRotation))
 	{
-		FIKRigTarget Target;
-		if (GetEffectorTarget(SolverDef->TransformTarget, Target))
-		{
-			int32 Index = InOutGlobalTransform.Hierarchy->GetIndex(SolverDef->TransformTarget.Bone);
-			if (Index != INDEX_NONE)
-			{
-				FTransform CurrentTransform = InOutGlobalTransform.GetGlobalTransform(Index);
-
-				if (SolverDef->bEnablePosition)
-				{
-					CurrentTransform.SetLocation(Target.PositionTarget.Position);
-				}
-				if (SolverDef->bEnableRotation)
-				{
-					CurrentTransform.SetRotation(Target.RotationTarget.Rotation.Quaternion());
-				}
-
-				InOutGlobalTransform.SetGlobalTransform(Index, CurrentTransform, true);
-			}
-		}
+		UpdateEffectors();
 	}
-}
\ No newline at end of file
+
+	Super::PostEditChangeChainProperty(PropertyChangedEvent);
+}
+
+#endif// WITH_EDITOR
\ No newline at end of file
diff --git a/Engine/Plugins/Animation/IKRig/Source/IKRig/Private/Solvers/TransformSolverDefinition.cpp b/Engine/Plugins/Animation/IKRig/Source/IKRig/Private/Solvers/TransformSolverDefinition.cpp
deleted file mode 100644
diff --git a/Engine/Plugins/Animation/IKRig/Source/IKRig/Public/IKRigConstraintDefinition.h b/Engine/Plugins/Animation/IKRig/Source/IKRig/Public/IKRigConstraintDefinition.h
deleted file mode 100644
diff --git a/Engine/Plugins/Animation/IKRig/Source/IKRig/Public/IKRigDefinition.h b/Engine/Plugins/Animation/IKRig/Source/IKRig/Public/IKRigDefinition.h
--- a/Engine/Plugins/Animation/IKRig/Source/IKRig/Public/IKRigDefinition.h
+++ b/Engine/Plugins/Animation/IKRig/Source/IKRig/Public/IKRigDefinition.h
@@ -1,23 +1,16 @@
 // Copyright Epic Games, Inc. All Rights Reserved.
 
-/**
- * Contains IKRig Definition 
- *
- *  https://docs.google.com/document/d/1yd8GCfT2aufxSdb5jAzlNTr1SptxEFpS9pWdQY-8LIk/edit#
- */
-
 #pragma once
 
 #include "CoreMinimal.h"
 #include "UObject/ObjectMacros.h"
 #include "UObject/Object.h"
-//#include "UObject/WeakObjectPtr.h"
+#include "IKRigConstraint.h"
 #include "IKRigDataTypes.h"
 #include "IKRigHierarchy.h"
 #include "IKRigDefinition.generated.h"
 
-class UIKRigSolverDefinition;
-class UIKRigConstraintDefinition;
+class UIKRigSolver;
 
 UCLASS(Blueprintable)
 class IKRIG_API UIKRigDefinition : public UObject
@@ -25,15 +18,21 @@ class IKRIG_API UIKRigDefinition : public UObject
 	GENERATED_BODY()
 
 public:
+
 	UIKRigDefinition();
 
-	// ensure the goal name is unique
-	void EnsureCreateUniqueGoalName(FName& InOutGoal) const;
-public: 
 	/** Source Asset imported */
 	UPROPERTY()
 	TSoftObjectPtr<UObject> SourceAsset;
 
+	UPROPERTY(EditAnywhere, instanced, Category = Constraints)
+	TArray<UIKRigConstraint*> Constraints;
+
+public:
+
+	// ensure the goal name is unique
+	void EnsureCreateUniqueGoalName(FName& InOutGoal) const;
+
 	// goal related APIs
 	const TMap<FName, FIKRigGoal>& GetGoals() const
 	{
@@ -52,9 +51,9 @@ public:
 		return ReferencePose;
 	}
 
-	const TArray<UIKRigSolverDefinition*>& GetSolverDefinitions() const 
+	const TArray<UIKRigSolver*>& GetSolvers() const 
 	{
-		return SolverDefinitions;
+		return Solvers;
 	}
 
 	const FIKRigHierarchy& GetHierarchy() const
@@ -84,18 +83,12 @@ private:
 	UPROPERTY(VisibleAnywhere, Category = "Hierarchy")
 	FIKRigTransform ReferencePose;
 
-	/*********** Solver Definition *********************/
-	// List of solvers, in the order of operations
+	// stack of solvers, executed in order
 	UPROPERTY(EditAnywhere, instanced, Category = "Solver")
-	TArray<UIKRigSolverDefinition*> SolverDefinitions;
-
-public:
-	/*********** Constraint Definition *********************/
-	// contains constraint data
-	UPROPERTY(EditAnywhere, instanced, Category = "Constraint")
-	UIKRigConstraintDefinition* ConstraintDefinitions;
+	TArray<UIKRigSolver*> Solvers;
 
 private:
+
 	/*********** Goals with Default Value******************/
 	// goals data
 	// this is cached by "Unique Internal Name" here
@@ -110,6 +103,7 @@ private:
 	virtual void PostLoad() override;
 	virtual void BeginDestroy() override;
 	// END UObject functions
+
 #if WITH_EDITOR
 	// https ://drive.google.com/file/d/1EO0Ijojx-0jommdHuZv1JYxfCdoFLd44/view
 
@@ -135,7 +129,7 @@ private:
 
 	// ensure it's sorted from parent to children
 	void EnsureSortedCorrectly(bool bReSortIfNeeded=false);
-	void Sanitize();
+	void Sanitize(); // hopefully we can get rid of this
 
 	// TODO: delegate for hierarchy changes
 	DECLARE_MULTICAST_DELEGATE_OneParam(FBoneAdded, FName /*BoneName*/);
diff --git a/Engine/Plugins/Animation/IKRig/Source/IKRig/Public/IKRigSolver.h b/Engine/Plugins/Animation/IKRig/Source/IKRig/Public/IKRigSolver.h
--- a/Engine/Plugins/Animation/IKRig/Source/IKRig/Public/IKRigSolver.h
+++ b/Engine/Plugins/Animation/IKRig/Source/IKRig/Public/IKRigSolver.h
@@ -1,22 +1,17 @@
 // Copyright Epic Games, Inc. All Rights Reserved.
 
-/**
- * Contains IKRig Solver Definition 
- *
- */
-
 #pragma once
 
 #include "CoreMinimal.h"
 #include "UObject/ObjectMacros.h"
 #include "UObject/Object.h"
+#include "IKRigDataTypes.h"
 #include "IKRigSolver.generated.h"
 
 struct FIKRigEffector;
 struct FIKRigTarget;
 struct FIKRigTransform;
 struct FIKRigTransformModifier;
-class UIKRigSolverDefinition;
 struct FControlRigDrawInterface;
 
 // run time processor 
@@ -29,38 +24,66 @@ public:
 	/** Required delegates to run this solver */
 	DECLARE_DELEGATE_RetVal(const FIKRigTransform&, FIKRigTransformGetter);
 	DECLARE_DELEGATE_RetVal_TwoParams(bool, FIKRigGoalGetter, const FName& /*InGoalName*/, FIKRigTarget& /*OutTarget*/);
-//	DECLARE_DELEGATE_TwoParams(FSolveConstraint, TArray<FName>& /*Constraint*/, FIKRigTransform& /*CurrentTransform*/);
-
-protected:
-
-	UPROPERTY(VisibleAnywhere, Category = "Definition")
-	UIKRigSolverDefinition* SolverDefinition;
-
-public:
-	UPROPERTY(EditAnywhere, Category = "Definition")
-	bool bEnabled = true;
 
 	// input hierarchy and ref pose? 
-	void Init(UIKRigSolverDefinition* InSolverDefinition, const FIKRigTransformModifier& TransformModifier, FIKRigTransformGetter InRefPoseGetter, FIKRigGoalGetter InGoalGetter/*, FSolveConstraint& InConstraintHandler*/);
+	void Init(const FIKRigTransformModifier& TransformModifier, FIKRigTransformGetter InRefPoseGetter, FIKRigGoalGetter InGoalGetter/*, FSolveConstraint& InConstraintHandler*/);
 
 	// input : goal getter or goals
 	// output : modified pose - GlobalTransforms
 	// or use SolverInternal function
 	void Solve(FIKRigTransformModifier& InOutGlobalTransform, FControlRigDrawInterface* InOutDrawInterface);
 
+	const TIKRigEffectorMap<FName>& GetEffectorToGoal() const
+	{
+		return EffectorToGoal;
+	}
+
+	void CollectGoals(TArray<FName>& OutGoals);
+
 protected:
-	// during we don't mutate bone transform, but you can read it
+
 	virtual void InitInternal(const FIKRigTransformModifier& InGlobalTransform) {};
 	virtual void SolveInternal(FIKRigTransformModifier& InOutGlobalTransform, FControlRigDrawInterface* InOutDrawInterface) {};
 	virtual bool IsSolverActive() const;
 
 	bool GetEffectorTarget(const FIKRigEffector& InEffector, FIKRigTarget& OutTarget) const;
 	const FIKRigTransform& GetReferencePose() const;
+
+	UPROPERTY(EditAnywhere, Category = "Definition")
+	bool bEnabled = true;
+
+	// effector name to goals name map
+	TIKRigEffectorMap<FName> EffectorToGoal;
+
 private:
+
 	// delegate
 	FIKRigTransformGetter RefPoseGetter;
 	FIKRigGoalGetter GoalGetter;
-	
-//	FSolveConstraint ConstraintHandler;
+
+	/// BEGIN UObject
+	virtual void PostLoad() override;
+	virtual void Serialize(FArchive& Ar) override;
+	/// END UObject
+
+	friend class UIKRigController;
+
+#if WITH_EDITOR
+private:
+	void RenameGoal(const FName& OldName, const FName& NewName);
+	// get unique name delegate by IKRigDefinition
+	void EnsureUniqueGoalName(FName& InOutGoalName) const;
+
+	DECLARE_MULTICAST_DELEGATE(FGoalNeedsUpdate);
+	FGoalNeedsUpdate GoalNeedsUpdateDelegate;
+
+protected:
+	FName CreateUniqueGoalName(const TCHAR* Suffix) const;
+	void OnGoalHasBeenUpdated();
+	void EnsureToAddEffector(const FIKRigEffector& InEffector, const FString& InPrefix);
+	void EnsureToRemoveEffector(const FIKRigEffector& InEffector);
+
+	virtual void UpdateEffectors() {};
+#endif // WITH_EDITOR
 };
 
diff --git a/Engine/Plugins/Animation/IKRig/Source/IKRig/Public/IKRigSolverDefinition.h b/Engine/Plugins/Animation/IKRig/Source/IKRig/Public/IKRigSolverDefinition.h
deleted file mode 100644
diff --git a/Engine/Plugins/Animation/IKRig/Source/IKRig/Public/Solvers/ConstraintSolver.h b/Engine/Plugins/Animation/IKRig/Source/IKRig/Public/Solvers/ConstraintSolver.h
--- a/Engine/Plugins/Animation/IKRig/Source/IKRig/Public/Solvers/ConstraintSolver.h
+++ b/Engine/Plugins/Animation/IKRig/Source/IKRig/Public/Solvers/ConstraintSolver.h
@@ -8,7 +8,6 @@
 #pragma once
 
 #include "IKRigSolver.h"
-#include "IKRigConstraintDefinition.h"
 #include "ConstraintSolver.generated.h"
 
 class UIKRigConstraint;
@@ -22,35 +21,32 @@ class IKRIG_API UIKRigConstraintSolver : public UIKRigSolver
 
 	DECLARE_DELEGATE_TwoParams(FIKRigQueryConstraint, const FName& InConstraintName, UIKRigConstraint& OutConstraint);
 
-	UPROPERTY(transient)
-	UIKRigConstraintDefinition* ConstraintDefinition;
 
 public: 
 	// during we don't mutate bone transform, but you can read it
 	virtual void InitInternal(const FIKRigTransformModifier& InGlobalTransform) override;
 	virtual void SolveInternal(FIKRigTransformModifier& InOutGlobalTransform, FControlRigDrawInterface* InOutDrawInterface) override;
-	virtual bool IsSolverActive() const override;
 
 	// this can be utilized for just subset of constraints
-	void SolverConstraints(const TArray<FName>& ConstraintsList, FIKRigTransformModifier& InOutGlobalTransform, FControlRigDrawInterface* InOutDrawInterface);
+	//void SolverConstraints(const TArray<FName>& ConstraintsList, FIKRigTransformModifier& InOutGlobalTransform, FControlRigDrawInterface* InOutDrawInterface);
 
 	// register/unregister query function
-	void RegisterQueryConstraintHandler(const FIKRigQueryConstraint& InQueryConstraintHandler);
-	void UnregisterQueryConstraintHandler();
+	//void RegisterQueryConstraintHandler(const FIKRigQueryConstraint& InQueryConstraintHandler);
+	//void UnregisterQueryConstraintHandler();
 
-	void SetConstraintDefinition(UIKRigConstraintDefinition* InConstraintDefinition);
+	//void SetConstraintDefinition(UIKRigConstraintDefinition* InConstraintDefinition);
 
 private:
 	// delegate
-	FIKRigQueryConstraint QueryConstraintHandler;
+	//FIKRigQueryConstraint QueryConstraintHandler;
 
 	// current active profile
-	UPROPERTY(transient)
-	FName ActiveProfile;
+	//UPROPERTY(transient)
+	//FName ActiveProfile;
 
 	// instanced information you can mutate from ConstraintDefinition
 	// we want uproperty, so that it doesn't GC-ed
-	UPROPERTY(transient)
-	TMap<FName, FIKRigConstraintProfile> ConstraintProfiles;
+	//UPROPERTY(transient)
+	//TMap<FName, FIKRigConstraintProfile> ConstraintProfiles;
 };
 
diff --git a/Engine/Plugins/Animation/IKRig/Source/IKRig/Public/Solvers/TransformSolver.h b/Engine/Plugins/Animation/IKRig/Source/IKRig/Public/Solvers/TransformSolver.h
--- a/Engine/Plugins/Animation/IKRig/Source/IKRig/Public/Solvers/TransformSolver.h
+++ b/Engine/Plugins/Animation/IKRig/Source/IKRig/Public/Solvers/TransformSolver.h
@@ -1,25 +1,44 @@
 // Copyright Epic Games, Inc. All Rights Reserved.
 
-/**
- * Contains Transform Solver Execution
- *
- */
-
 #pragma once
 
 #include "IKRigSolver.h"
 #include "TransformSolver.generated.h"
 
 
-// run time for UTransformSolverDefinition
 UCLASS()
 class IKRIG_API UTransformSolver : public UIKRigSolver
 {
 	GENERATED_BODY()
 
+public:
+
+	UTransformSolver();
+
+	UPROPERTY(EditAnywhere, Category = "Solver")
+	bool bEnablePosition = true;
+
+	UPROPERTY(EditAnywhere, Category = "Solver")
+	bool bEnableRotation = true;
+
+	UPROPERTY(EditAnywhere, Category = "Solver")
+	FIKRigEffector TransformTarget;
+
 protected:
+	
 	virtual void InitInternal(const FIKRigTransformModifier& InGlobalTransform) override;
 	virtual void SolveInternal(FIKRigTransformModifier& InOutGlobalTransform, FControlRigDrawInterface* InOutDrawInterface) override;
 	virtual bool IsSolverActive() const override;
+
+private:
+#if WITH_EDITOR
+	virtual void UpdateEffectors() override;
+
+	// UObject interface
+	virtual void PostEditChangeChainProperty(struct FPropertyChangedChainEvent& PropertyChangedEvent) override;
+	// UObject interface
+#endif // WITH_EDITOR
+
+	const FString TransformTargetName;
 };
 
diff --git a/Engine/Plugins/Animation/IKRig/Source/IKRig/Public/Solvers/TransformSolverDefinition.h b/Engine/Plugins/Animation/IKRig/Source/IKRig/Public/Solvers/TransformSolverDefinition.h
deleted file mode 100644
diff --git a/Engine/Plugins/Animation/IKRig/Source/IKRigEditor/Private/IKRigController.cpp b/Engine/Plugins/Animation/IKRig/Source/IKRigEditor/Private/IKRigController.cpp
--- a/Engine/Plugins/Animation/IKRig/Source/IKRigEditor/Private/IKRigController.cpp
+++ b/Engine/Plugins/Animation/IKRig/Source/IKRigEditor/Private/IKRigController.cpp
@@ -1,16 +1,10 @@
 // Copyright Epic Games, Inc. All Rights Reserved.
 
-/**
- * IKController implementation class
- *
- */
-
 #include "IKRigController.h"
 #include "IKRigDefinition.h"
 #include "AnimationRuntime.h"
-#include "IKRigSolverDefinition.h"
+#include "IKRigSolver.h"
 #include "ScopedTransaction.h"
-#include "IKRigConstraintDefinition.h"
 #include "IKRigConstraint.h"
 
 #define LOCTEXT_NAMESPACE	"IKRigController"
@@ -54,7 +48,7 @@ void UIKRigController::AddReferencedObjects(UObject* InThis, FReferenceCollector
 	{
 		for (auto Iter = DefinitionToControllerMap.CreateIterator(); Iter; ++Iter)
 		{
-			// we add controllers to the list. so that it doens't GCed
+			// we add controllers to the list. so that it doesn't GCed
 			// when do we delete from the list?
 			UIKRigController* Obj = Iter.Value();
 			Collector.AddReferencedObject(Obj);
@@ -66,9 +60,9 @@ void UIKRigController::BeginDestroy()
 {
 	if (IKRigDefinition)
 	{
-		for (int32 Index = 0; Index < IKRigDefinition->SolverDefinitions.Num(); ++Index)
+		for (int32 Index = 0; Index < IKRigDefinition->Solvers.Num(); ++Index)
 		{
-			UninitializeIKRigSolverDefinition(IKRigDefinition->SolverDefinitions[Index]);
+			UninitializeSolver(IKRigDefinition->Solvers[Index]);
 		}
 	}
 
@@ -80,9 +74,9 @@ void UIKRigController::SetIKRigDefinition(UIKRigDefinition* InIKRigDefinition)
 {
 	if (IKRigDefinition)
 	{
-		for (int32 Index = 0; Index < IKRigDefinition->SolverDefinitions.Num(); ++Index)
+		for (int32 Index = 0; Index < IKRigDefinition->Solvers.Num(); ++Index)
 		{
-			UninitializeIKRigSolverDefinition(IKRigDefinition->SolverDefinitions[Index]);
+			UninitializeSolver(IKRigDefinition->Solvers[Index]);
 		}
 	}
 
@@ -90,35 +84,35 @@ void UIKRigController::SetIKRigDefinition(UIKRigDefinition* InIKRigDefinition)
 
 	if (IKRigDefinition)
 	{
-		for (int32 Index = 0; Index < IKRigDefinition->SolverDefinitions.Num(); ++Index)
+		for (int32 Index = 0; Index < IKRigDefinition->Solvers.Num(); ++Index)
 		{
-			InitializeIKRigSolverDefinition(IKRigDefinition->SolverDefinitions[Index]);
+			InitializeSolver(IKRigDefinition->Solvers[Index]);
 		}
 	}
 }
 
-void UIKRigController::InitializeIKRigSolverDefinition(UIKRigSolverDefinition* SolverDef)
+void UIKRigController::InitializeSolver(UIKRigSolver* Solver)
 {
-	if (SolverDef)
+	if (Solver)
 	{
 		// we want to register this controller delegate
 		// if there are multiple controllers managing one IKRigDefinition
 		// we want to ensure that works 
-		FDelegateHandle DelegateHandle = SolverDef->GoalNeedsUpdateDelegate.AddUObject(this, &UIKRigController::UpdateGoal);
-		SolverDelegateHandles.Add(SolverDef) = DelegateHandle;
+		FDelegateHandle DelegateHandle = Solver->GoalNeedsUpdateDelegate.AddUObject(this, &UIKRigController::UpdateGoal);
+		SolverDelegateHandles.Add(Solver) = DelegateHandle;
 
-		SolverDef->UpdateEffectors();
+		Solver->UpdateEffectors();
 	}
 }
 
-void UIKRigController::UninitializeIKRigSolverDefinition(UIKRigSolverDefinition* SolverDef)
+void UIKRigController::UninitializeSolver(UIKRigSolver* Solver)
 {
-	if (SolverDef)
+	if (Solver)
 	{
-		FDelegateHandle* DelegateHandle = SolverDelegateHandles.Find(SolverDef);
+		FDelegateHandle* DelegateHandle = SolverDelegateHandles.Find(Solver);
 		if (ensure(DelegateHandle))
 		{
-			SolverDef->GoalNeedsUpdateDelegate.Remove(*DelegateHandle);
+			Solver->GoalNeedsUpdateDelegate.Remove(*DelegateHandle);
 		}
 	}
 }
@@ -208,19 +202,19 @@ void UIKRigController::ResetHierarchy()
 }
 
 // solver operators
-UIKRigSolverDefinition* UIKRigController::AddSolver(TSubclassOf<UIKRigSolverDefinition> InIKRigSolverDefinitionClass)
+UIKRigSolver* UIKRigController::AddSolver(TSubclassOf<UIKRigSolver> InIKRigSolverClass)
 {
 	if (IKRigDefinition)
 	{
 		FScopedTransaction Transaction(LOCTEXT("AddSolver_Label", "Add Solver"));
 		IKRigDefinition->Modify();
 
-		UIKRigSolverDefinition* NewSolver = NewObject<UIKRigSolverDefinition>(IKRigDefinition, InIKRigSolverDefinitionClass);
+		UIKRigSolver* NewSolver = NewObject<UIKRigSolver>(IKRigDefinition, InIKRigSolverClass);
 		check(NewSolver);
 
 		// todo: set delegate for the goal update
-		IKRigDefinition->SolverDefinitions.Add(NewSolver);
-		InitializeIKRigSolverDefinition(NewSolver);
+		IKRigDefinition->Solvers.Add(NewSolver);
+		InitializeSolver(NewSolver);
 		IKRigDefinition->UpdateGoal();
 		return NewSolver;
 	}
@@ -228,42 +222,42 @@ UIKRigSolverDefinition* UIKRigController::AddSolver(TSubclassOf<UIKRigSolverDefi
 	return nullptr;
 }
 
-int32 UIKRigController::GetTotalSolverCount() const
+int32 UIKRigController::GetNumSolvers() const
 {
 	if (IKRigDefinition)
 	{
-		return IKRigDefinition->SolverDefinitions.Num();
+		return IKRigDefinition->Solvers.Num();
 	}
 
 	return 0;
 }
 
-UIKRigSolverDefinition* UIKRigController::GetSolver(int32 Index) const
+UIKRigSolver* UIKRigController::GetSolver(int32 Index) const
 {
-	if (IKRigDefinition && IKRigDefinition->SolverDefinitions.IsValidIndex(Index))
+	if (IKRigDefinition && IKRigDefinition->Solvers.IsValidIndex(Index))
 	{
-		return IKRigDefinition->SolverDefinitions[Index];
+		return IKRigDefinition->Solvers[Index];
 	}
 
 	return nullptr;
 }
 
-void UIKRigController::RemoveSolver(UIKRigSolverDefinition* SolverToDelete)
+void UIKRigController::RemoveSolver(UIKRigSolver* SolverToDelete)
 {
 	if (IKRigDefinition && SolverToDelete)
 	{
 		FScopedTransaction Transaction(LOCTEXT("RemoveSolver_Label", "Remove Solver"));
 		IKRigDefinition->Modify();
 
-		UninitializeIKRigSolverDefinition(SolverToDelete);
-		IKRigDefinition->SolverDefinitions.Remove(SolverToDelete);
+		UninitializeSolver(SolverToDelete);
+		IKRigDefinition->Solvers.Remove(SolverToDelete);
 		IKRigDefinition->UpdateGoal();
 	}
 }
 
-bool UIKRigController::ValidateSolver(UIKRigSolverDefinition* const SolverDef) const
+bool UIKRigController::ValidateSolver(UIKRigSolver* const Solver) const
 {
-	return (IKRigDefinition && SolverDef && IKRigDefinition->SolverDefinitions.Find(SolverDef) != INDEX_NONE);
+	return (IKRigDefinition && Solver && IKRigDefinition->Solvers.Find(Solver) != INDEX_NONE);
 }
 
 void UIKRigController::UpdateGoal()
@@ -278,11 +272,11 @@ void UIKRigController::UpdateGoal()
 	}
 }
 
-FName UIKRigController::GetGoalName(UIKRigSolverDefinition* InSolverDefinition, const FIKRigEffector& InEffector)
+FName UIKRigController::GetGoalName(UIKRigSolver* InSolver, const FIKRigEffector& InEffector)
 {
-	if (InSolverDefinition)
+	if (InSolver)
 	{
-		FName* GoalName = InSolverDefinition->EffectorToGoal.Find(InEffector);
+		FName* GoalName = InSolver->EffectorToGoal.Find(InEffector);
 		if (GoalName)
 		{
 			return *GoalName;
@@ -292,14 +286,14 @@ FName UIKRigController::GetGoalName(UIKRigSolverDefinition* InSolverDefinition,
 	return NAME_None;
 }
 
-void UIKRigController::SetGoalName(UIKRigSolverDefinition* InSolverDefinition, const FIKRigEffector& InEffector, const FName& NewGoalName)
+void UIKRigController::SetGoalName(UIKRigSolver* InSolver, const FIKRigEffector& InEffector, const FName& NewGoalName)
 {
-	if (InSolverDefinition)
+	if (InSolver)
 	{
 		FScopedTransaction Transaction(LOCTEXT("SetGoalName_Label", "Set Goal Name"));
-		InSolverDefinition->Modify();
+		InSolver->Modify();
 
-		FName* GoalName = InSolverDefinition->EffectorToGoal.Find(InEffector);
+		FName* GoalName = InSolver->EffectorToGoal.Find(InEffector);
 		if (GoalName)
 		{
 			*GoalName = NewGoalName;
@@ -309,169 +303,22 @@ void UIKRigController::SetGoalName(UIKRigSolverDefinition* InSolverDefinition, c
 	}
 }
 
-void UIKRigController::AutoConfigure(UIKRigSolverDefinition* SolverDef)
-{
-	if (ValidateSolver(SolverDef))
-	{
-		FScopedTransaction Transaction(LOCTEXT("AutoConfigure_Label", "Auto Congirure"));
-		SolverDef->Modify();
-
-		SolverDef->AutoConfigure();
-	}
-}
-
-bool UIKRigController::CanAutoConfigure(UIKRigSolverDefinition* SolverDef) const
-{
-	if (ValidateSolver(SolverDef))
-	{
-		return SolverDef->CanAutoConfigure();
-	}
-
-	return false;
-}
-
-// constraint operators
-// create new profile
-void UIKRigController::CreateNewProfile(FName& InNewProfileName)
-{
-	if (IKRigDefinition && InNewProfileName != NAME_None)
-	{
-		FScopedTransaction Transaction(LOCTEXT("CreateNewProfile_Label", "Create New Constraint Profile"));
-		IKRigDefinition->Modify();
-		// we just find or add
-		IKRigDefinition->ConstraintDefinitions->ConstraintProfiles.FindOrAdd(InNewProfileName);
-	}
-}
-
-bool UIKRigController::RemoveConstraintProfile(const FName& InProfileName)
-{
-	if (IKRigDefinition && InProfileName != NAME_None)
-	{
-		FScopedTransaction Transaction(LOCTEXT("RemoveProfile_Label", "Remove Constraint Profile"));
-		IKRigDefinition->Modify();
-		// we just find or add
-		return IKRigDefinition->ConstraintDefinitions->ConstraintProfiles.Remove(InProfileName) > 0;
-
-	}
-
-	return false;
-}
-
-void UIKRigController::RenameProfile(FName InCurrentProfileName, FName& InNewProfileName)
-{
-	// write this 
-	ensure (false);
-	// we don't allow change default profile until we support saving the default profile
-}
-
-FIKRigConstraintProfile* UIKRigController::GetConstraintProfile(const FName& InProfileName) const
-{
-	if (IKRigDefinition)
-	{
-		FName SearchProfileName = (InProfileName == NAME_None)? UIKRigConstraintDefinition::DefaultProfileName : InProfileName;
-		return IKRigDefinition->ConstraintDefinitions->ConstraintProfiles.Find(SearchProfileName);
-	}
-
-	return nullptr;
-}
-
-void UIKRigController::EnsureUniqueConstraintName(FName& InOutName)
+UIKRigConstraint* UIKRigController::AddConstraint(TSubclassOf<UIKRigConstraint> NewConstraintType)
 {
-	int32 Index = 1;
-
-	TArray<FName> ConstraintNames;
-
-	GetConstraintNames(ConstraintNames);
+	
+	FScopedTransaction Transaction(LOCTEXT("AddConstraint_Label", "Add Constraint"));
+	IKRigDefinition->Modify();
 
-	FString NewName = InOutName.ToString();
-	while (ConstraintNames.Contains(FName(*NewName)))
+	UIKRigConstraint* NewRigConstraint = NewObject<UIKRigConstraint>(IKRigDefinition, NewConstraintType);
+	if (NewRigConstraint)
 	{
-		NewName = FString::Format(TEXT("{0}_{1}"), { InOutName.ToString(), Index++ });
-	}
-
-	InOutName = FName(*NewName);
-}
-
-UIKRigConstraint* UIKRigController::AddConstraint(TSubclassOf<UIKRigConstraint> NewConstraintType, FName& InOutNewName, FName InProfile /*= NAME_None*/)
-{
-	FIKRigConstraintProfile* Profile = GetConstraintProfile(InProfile);
-	if (Profile)
-	{
-		FScopedTransaction Transaction(LOCTEXT("AddConstraint_Label", "Add Constraint"));
-		IKRigDefinition->Modify();
-
-		EnsureUniqueConstraintName(InOutNewName);
-		UIKRigConstraint* NewRigConstraint = NewObject<UIKRigConstraint>(IKRigDefinition, NewConstraintType, InOutNewName);
-		if (NewRigConstraint)
-		{
-		
-			Profile->Constraints.Add(NewRigConstraint->GetFName(), NewRigConstraint);
-			return NewRigConstraint;
-		}
+		IKRigDefinition->Constraints.Add(NewRigConstraint);
+		return NewRigConstraint;
 	}
+	
 	return nullptr;	
 }
 
-UIKRigConstraint* UIKRigController::GetConstraint(const FName& InProfileName, const FName& InName) const
-{
-	FIKRigConstraintProfile* Profile = GetConstraintProfile(InProfileName);
-	if (Profile)
-	{
-		UIKRigConstraint** RigConstraint = Profile->Constraints.Find(InName);
-		if (RigConstraint)
-		{
-			return *RigConstraint;
-		}
-	}
-
-	return nullptr;
-}
-
-bool UIKRigController::RemoveConstraint(const FName& InConstraintName)
-{
-	// remove this in all profile
-	if (IKRigDefinition)
-	{
-		FScopedTransaction Transaction(LOCTEXT("AddConstraint_Label", "Add Constraint"));
-		IKRigDefinition->Modify();
-
-		for (auto Iter = IKRigDefinition->ConstraintDefinitions->ConstraintProfiles.CreateIterator(); Iter; ++Iter)
-		{
-			FIKRigConstraintProfile& Profile = Iter.Value();
-			for (auto InnerIter = Profile.Constraints.CreateIterator(); InnerIter; ++InnerIter)
-			{
-				// possible the constraint name may not be there
-				// what this means, when we rename, we have to rename all of them
-				Profile.Constraints.Remove(InConstraintName);
-			}
-		}
-		// remove return?
-		return true;
-	}
-
-	return false;
-}
-
-void UIKRigController::GetConstraintProfileNames(TArray<FName>& OutProfileNames) const
-{
-	if (IKRigDefinition)
-	{
-		OutProfileNames.Reset();
-		IKRigDefinition->ConstraintDefinitions->ConstraintProfiles.GenerateKeyArray(OutProfileNames);
-	}
-}
-
-void UIKRigController::GetConstraintNames(TArray<FName>& OutConstraintNames) const
-{
-	// still can be null if IKRigDefinition is nullptr
-	FIKRigConstraintProfile* DefaultProfile = GetConstraintProfile(NAME_None);
-	if (DefaultProfile)
-	{
-		OutConstraintNames.Reset();
-		DefaultProfile->Constraints.GenerateKeyArray(OutConstraintNames);
-	}
-}
-
 // goal operators
 void UIKRigController::QueryGoals(TArray<FName>& OutGoals) const
 {
@@ -489,11 +336,11 @@ void UIKRigController::RenameGoal(const FName& OldName, const FName& NewName)
 		FScopedTransaction Transaction(LOCTEXT("RenameGoal_Label", "Rename Goal"));
 		IKRigDefinition->Modify();
 
-		for (UIKRigSolverDefinition* SolverDef : IKRigDefinition->SolverDefinitions)
+		for (UIKRigSolver* Solver : IKRigDefinition->Solvers)
 		{
-			if (SolverDef)
+			if (Solver)
 			{
-				SolverDef->RenameGoal(OldName, NewName);
+				Solver->RenameGoal(OldName, NewName);
 			}
 		}
 		
diff --git a/Engine/Plugins/Animation/IKRig/Source/IKRigEditor/Private/IKRigDefinitionDetails.cpp b/Engine/Plugins/Animation/IKRig/Source/IKRigEditor/Private/IKRigDefinitionDetails.cpp
--- a/Engine/Plugins/Animation/IKRig/Source/IKRigEditor/Private/IKRigDefinitionDetails.cpp
+++ b/Engine/Plugins/Animation/IKRig/Source/IKRigEditor/Private/IKRigDefinitionDetails.cpp
@@ -9,7 +9,7 @@
 
 #include "IKRigDefinition.h"
 #include "IKRigController.h"
-#include "IKRigSolverDefinition.h"
+#include "IKRigSolver.h"
 #include "IKRigConstraint.h"
 
 #include "ScopedTransaction.h"
@@ -253,7 +253,7 @@ UClass* SelectClass(UClass* ClassType, const FText& TitleText)
 
 FReply FIKRigDefinitionDetails::OnShowSolverClassPicker()
 {
-	UClass* ChosenClass = SelectClass(UIKRigSolverDefinition::StaticClass(), LOCTEXT("SelectSolverClass", "Select Solver Class"));
+	UClass* ChosenClass = SelectClass(UIKRigSolver::StaticClass(), LOCTEXT("SelectSolverClass", "Select Solver Class"));
 	if (ChosenClass)
 	{
 		IKRigController->AddSolver(ChosenClass);
@@ -267,9 +267,7 @@ FReply FIKRigDefinitionDetails::OnShowConstraintClassPicker()
 	UClass* ChosenClass = SelectClass(UIKRigConstraint::StaticClass(), LOCTEXT("SelectConstraintClass", "Select Constraint Class"));
 	if (ChosenClass)
 	{
-		// ask for name of constraint
-		FName NewName(TEXT("NewConstraint"));
-		IKRigController->AddConstraint(ChosenClass, NewName);
+		IKRigController->AddConstraint(ChosenClass);
 	}
 
 	return FReply::Handled();
diff --git a/Engine/Plugins/Animation/IKRig/Source/IKRigEditor/Private/IKRigEffectorCustomization.cpp b/Engine/Plugins/Animation/IKRig/Source/IKRigEditor/Private/IKRigEffectorCustomization.cpp
--- a/Engine/Plugins/Animation/IKRig/Source/IKRigEditor/Private/IKRigEffectorCustomization.cpp
+++ b/Engine/Plugins/Animation/IKRig/Source/IKRigEditor/Private/IKRigEffectorCustomization.cpp
@@ -6,7 +6,7 @@
 #include "IDetailChildrenBuilder.h"
 
 #include "IKRigController.h"
-#include "IKRigSolverDefinition.h"
+#include "IKRigSolver.h"
  
 #include "PropertyEditorModule.h"
 #include "PropertyHandle.h"
@@ -50,7 +50,7 @@ void FIKRigEffectorCustomization::CustomizeHeader(TSharedRef<IPropertyHandle> St
 
 		Effector= reinterpret_cast<FIKRigEffector*>(Data);
 		// give me the name of goal by effecto
-		FName CurrentGoalName = IKRigController->GetGoalName(IKRigSolverDefinition.Get(), *Effector);
+		FName CurrentGoalName = IKRigController->GetGoalName(IKRigSolver.Get(), *Effector);
 
 		// go through profile and see if it has mine
 		for (const FName& Goal : Goals)
@@ -116,32 +116,31 @@ void FIKRigEffectorCustomization::SetSolverDefinition(TSharedRef<IPropertyHandle
 	TArray<UObject*> Objects;
 	StructPropertyHandle->GetOuterObjects(Objects);
 
-	IKRigSolverDefinition.Reset();
+	IKRigSolver.Reset();
 	IKRigController.Reset();
 
-	UIKRigSolverDefinition* SelectedSolverDefinition = nullptr;
+	UIKRigSolver* SelectedSolver = nullptr;
 	
 	// currently not allowing multi selection
 	// or if you select different IKRigSolvers, this will break
 	for (UObject* Outer : Objects)
 	{
-		if (UIKRigSolverDefinition* SolverDefinition = Cast<UIKRigSolverDefinition>(Outer))
+		if (UIKRigSolver* Solver = Cast<UIKRigSolver>(Outer))
 		{
-			SelectedSolverDefinition = SolverDefinition;
+			SelectedSolver = Solver;
 			break;
 		}
 	}
 
-	IKRigSolverDefinition = SelectedSolverDefinition;
+	IKRigSolver = SelectedSolver;
 
-	// once we know solver definion, we know outer is IKRigDefinition
-	if (SelectedSolverDefinition)
+	// once we know solver, we know outer is IKRigDefinition
+	if (SelectedSolver)
 	{
-		UIKRigDefinition* IKRigDefinition = Cast<UIKRigDefinition>(SelectedSolverDefinition->GetOuter());
+		UIKRigDefinition* IKRigDefinition = Cast<UIKRigDefinition>(SelectedSolver->GetOuter());
 		if (IKRigDefinition)
 		{
 			IKRigController = UIKRigController::GetControllerByRigDefinition(IKRigDefinition);
-
 			IKRigController->OnGoalModified.AddSP(this, &FIKRigEffectorCustomization::RefreshEffectorGoals);
 		}
 	}
@@ -190,17 +189,17 @@ void FIKRigEffectorCustomization::SetPropertyHandle(TSharedRef<IPropertyHandle>
 void FIKRigEffectorCustomization::SetNewGoalName(FName Name)
 {
 	// I have to change the mapping of SolverDefinition
-	if (UIKRigSolverDefinition* SolverDef = IKRigSolverDefinition.Get())
+	if (UIKRigSolver* Solver = IKRigSolver.Get())
 	{
-		IKRigController->SetGoalName(SolverDef, *Effector, Name);
+		IKRigController->SetGoalName(Solver, *Effector, Name);
 	}
 }
 
 FName FIKRigEffectorCustomization::GetSelectedEffectorGoal() const
 {
-	if (UIKRigSolverDefinition* SolverDef = IKRigSolverDefinition.Get())
+	if (UIKRigSolver* Solver = IKRigSolver.Get())
 	{
-		return IKRigController->GetGoalName(SolverDef, *Effector);
+		return IKRigController->GetGoalName(Solver, *Effector);
 	}
 
 	return NAME_None;
@@ -218,13 +217,13 @@ void FIKRigEffectorCustomization::OnEffectorGoalSelectionChanged(TSharedPtr<FStr
 
 void FIKRigEffectorCustomization::OnEffectorGoalChanged(const FText& InText, ETextCommit::Type CommitType)
 {
-	if (UIKRigSolverDefinition* SolverDef = IKRigSolverDefinition.Get())
+	if (UIKRigSolver* Solver = IKRigSolver.Get())
 	{
 		FString NewGoalName = InText.ToString();
 		NewGoalName.TrimStartAndEndInline();
 		if (!NewGoalName.IsEmpty())
 		{
-			IKRigController->SetGoalName(SolverDef, *Effector, FName(*NewGoalName));
+			IKRigController->SetGoalName(Solver, *Effector, FName(*NewGoalName));
 		}
 	}
 }
diff --git a/Engine/Plugins/Animation/IKRig/Source/IKRigEditor/Private/IKRigEffectorCustomization.h b/Engine/Plugins/Animation/IKRig/Source/IKRigEditor/Private/IKRigEffectorCustomization.h
--- a/Engine/Plugins/Animation/IKRig/Source/IKRigEditor/Private/IKRigEffectorCustomization.h
+++ b/Engine/Plugins/Animation/IKRig/Source/IKRigEditor/Private/IKRigEffectorCustomization.h
@@ -6,7 +6,7 @@
 #include "IPropertyTypeCustomization.h"
 #include "IKRigDataTypes.h"
 
-class UIKRigSolverDefinition;
+class UIKRigSolver;
 class UIKRigController;
 class SSearchableComboBox;
 
@@ -31,13 +31,13 @@ private:
 	TSharedPtr<IPropertyHandle> EffectorBoneProperty;
 
 	// Target Skeleton this widget is referencing
-	TWeakObjectPtr<UIKRigSolverDefinition>	IKRigSolverDefinition;
-	TWeakObjectPtr<UIKRigController>		IKRigController;
+	TWeakObjectPtr<UIKRigSolver>		IKRigSolver;
+	TWeakObjectPtr<UIKRigController>	IKRigController;
 
 	// Goal update option
-	TSharedPtr< SSearchableComboBox >			EffectorGoalComboBox;
-	TArray< TSharedPtr< FString > >				EffectorGoalComboList;
-	FIKRigEffector*								Effector;								
+	TSharedPtr<SSearchableComboBox>	EffectorGoalComboBox;
+	TArray<TSharedPtr<FString>>	EffectorGoalComboList;
+	FIKRigEffector*	Effector;								
 
 	TSharedRef<SWidget> MakeEffectorGoalComboWidget(TSharedPtr<FString> InItem);
 	// combo box
diff --git a/Engine/Plugins/Animation/IKRig/Source/IKRigEditor/Public/IKRigController.h b/Engine/Plugins/Animation/IKRig/Source/IKRigEditor/Public/IKRigController.h
--- a/Engine/Plugins/Animation/IKRig/Source/IKRigEditor/Public/IKRigController.h
+++ b/Engine/Plugins/Animation/IKRig/Source/IKRigEditor/Public/IKRigController.h
@@ -1,11 +1,5 @@
 // Copyright Epic Games, Inc. All Rights Reserved.
 
-/**
- * IKRigController
- *
- * Mutates IKRig asset side
- */
-
 #pragma once
 
 #include "CoreMinimal.h"
@@ -55,34 +49,22 @@ public:
 	// todo: more APIs for hierarchy operators
 
 	// solver operators
-	UIKRigSolverDefinition* AddSolver(TSubclassOf<UIKRigSolverDefinition> InIKRigSolverDefinitionClass);
-	int32 GetTotalSolverCount() const;
-	UIKRigSolverDefinition* GetSolver(int32 Index) const;
-	void RemoveSolver(UIKRigSolverDefinition* SolverToDelete);
-	void AutoConfigure(UIKRigSolverDefinition* SolverDef);
-	bool CanAutoConfigure(UIKRigSolverDefinition* SolverDef) const;
-
-	// constraint operators
-	// create new profile
-	void CreateNewProfile(FName& InNewProfileName);
-	bool RemoveConstraintProfile(const FName& InProfileName);
-	void RenameProfile(FName InCurrentProfileName, FName& InNewProfileName);
-
-	UIKRigConstraint* AddConstraint(TSubclassOf<UIKRigConstraint> NewConstraintType, FName& InOutNewName, FName InProfile=NAME_None);
-	UIKRigConstraint* GetConstraint(const FName& InProfileName, const FName& InName) const;
-	bool RemoveConstraint(const FName& InConstraintName);
-
-	void GetConstraintProfileNames(TArray<FName>& OutProfileNames) const;
-	void GetConstraintNames(TArray<FName>& OutConstraintNames) const;
+	UIKRigSolver* AddSolver(TSubclassOf<UIKRigSolver> InSolverClass);
+	void RemoveSolver(UIKRigSolver* SolverToDelete);
+	UIKRigSolver* GetSolver(int32 Index) const;
+	int32 GetNumSolvers() const;
+
+	// constraints
+	UIKRigConstraint* AddConstraint(TSubclassOf<UIKRigConstraint> NewConstraintType);
 
 	// goal operators
 	void QueryGoals(TArray<FName>& OutGoals) const;
-	// why do we need rename? @todo: remove - maybe in the fuutre when we want a tool that does rename automatically this could be useful, but for now it's not really
+	// why do we need rename? @todo: remove - maybe in the future when we want a tool that does rename automatically this could be useful, but for now it's not really
 	// goals are collected from solvers, so they could change name, and if it doesn't exists, it doesn't exist
 	void RenameGoal(const FName& OldName, const FName& NewName);
 
-	FName GetGoalName(UIKRigSolverDefinition* InSolverDefinition, const FIKRigEffector& InEffector);
-	void SetGoalName(UIKRigSolverDefinition* InSolverDefinition, const FIKRigEffector& InEffector, const FName& NewGoalName);
+	FName GetGoalName(UIKRigSolver* InSolver, const FIKRigEffector& InEffector);
+	void SetGoalName(UIKRigSolver* InSolver, const FIKRigEffector& InEffector, const FName& NewGoalName);
 
 	// this is to modify default value
 	FIKRigGoal* GetGoal(const FName& InGoalName);
@@ -105,19 +87,18 @@ private:
 	UPROPERTY(transient)
 	UIKRigDefinition* IKRigDefinition = nullptr;
 
-	bool ValidateSolver(UIKRigSolverDefinition* const SolverDef) const;
+	bool ValidateSolver(UIKRigSolver* const Solver) const;
 	void UpdateGoal();
 
 	FIKRigConstraintProfile* GetConstraintProfile(const FName& InProfileName) const;
 
-	void InitializeIKRigSolverDefinition(UIKRigSolverDefinition* SolverDef);
-	void UninitializeIKRigSolverDefinition(UIKRigSolverDefinition* SolverDef);
+	void InitializeSolver(UIKRigSolver* Solver);
+	void UninitializeSolver(UIKRigSolver* Solver);
 
-	TMap<UIKRigSolverDefinition*, FDelegateHandle> SolverDelegateHandles;
+	TMap<UIKRigSolver*, FDelegateHandle> SolverDelegateHandles;
 
 	static TMap<UIKRigDefinition*, UIKRigController*> DefinitionToControllerMap;
 
-	// IKRigDefinition set up
 	void SetIKRigDefinition(UIKRigDefinition* InIKRigDefinition);
 
 	void EnsureUniqueConstraintName(FName& InOutName);
diff --git a/Engine/Plugins/Experimental/FullBodyIK/Source/FullBodyIK/Private/FullBodyIKSolver.cpp b/Engine/Plugins/Experimental/FullBodyIK/Source/FullBodyIK/Private/FullBodyIKSolver.cpp
--- a/Engine/Plugins/Experimental/FullBodyIK/Source/FullBodyIK/Private/FullBodyIKSolver.cpp
+++ b/Engine/Plugins/Experimental/FullBodyIK/Source/FullBodyIK/Private/FullBodyIKSolver.cpp
@@ -9,11 +9,12 @@
 #include "IKRigHierarchy.h"
 #include "FBIKConstraint.h"
 #include "JacobianIK.h"
-#include "FullBodyIKSolverDefinition.h"
 #include "FBIKUtil.h"
 #include "FBIKShared.h"
 #include "Drawing/ControlRigDrawInterface.h"
 
+const FName UFullBodyIKSolver::EffectorTargetPrefix = FName(TEXT("FullBodyIKTarget"));
+
 //////////////////////////////////////////////
 // utility functions
 ////////////////////////////////////////////////
@@ -198,41 +199,40 @@ static void AddEffectors(const FIKRigTransformModifier& TransformModifier, const
 	}
 }
 /////////////////////////////////////////////////
+
+UFullBodyIKSolver::UFullBodyIKSolver()
+{
+}
+
 void UFullBodyIKSolver::InitInternal(const FIKRigTransformModifier& InGlobalTransform)
 {
-	if (UFullBodyIKSolverDefinition* SolverDef = Cast<UFullBodyIKSolverDefinition>(SolverDefinition))
-	{
-		LinkData.Reset();
-		EffectorTargets.Reset();
-		EffectorLinkIndices.Reset();
-		LinkDataToHierarchyIndices.Reset();
-		HierarchyToLinkDataMap.Reset();
-
-		// verify the chain
-		AddEffectors(InGlobalTransform, SolverDef->Root, SolverDef->Effectors, LinkData, EffectorTargets, EffectorLinkIndices, LinkDataToHierarchyIndices, HierarchyToLinkDataMap, SolverDef->SolverProperty);
-	}
+
+	LinkData.Reset();
+	EffectorTargets.Reset();
+	EffectorLinkIndices.Reset();
+	LinkDataToHierarchyIndices.Reset();
+	HierarchyToLinkDataMap.Reset();
+
+	// verify the chain
+	AddEffectors(InGlobalTransform, Root, Effectors, LinkData, EffectorTargets, EffectorLinkIndices, LinkDataToHierarchyIndices, HierarchyToLinkDataMap, SolverProperty);
+	
 }
 
 bool UFullBodyIKSolver::IsSolverActive() const
 {
 	if (Super::IsSolverActive())
 	{
- 		if (UFullBodyIKSolverDefinition* SolverDef = Cast<UFullBodyIKSolverDefinition>(SolverDefinition))
- 		{
-			// i'm not checkinf if this is valid, but we may not know for sure because 
- 			return SolverDef->Effectors.Num() > 1;
- 		}
+		return Effectors.Num() > 1;
 	}
 
 	return false;
 }
 
-void UFullBodyIKSolver::SolveInternal(FIKRigTransformModifier& InOutGlobalTransform, FControlRigDrawInterface* InOutDrawInterface)
+void UFullBodyIKSolver::SolveInternal(
+	FIKRigTransformModifier& InOutGlobalTransform, 
+	FControlRigDrawInterface* InOutDrawInterface)
 {
-
-	UFullBodyIKSolverDefinition* SolverDef = Cast<UFullBodyIKSolverDefinition>(SolverDefinition);
-
-	if (SolverDef && LinkDataToHierarchyIndices.Num() > 0)
+	if (LinkDataToHierarchyIndices.Num() > 0)
 	{
 #if 0
 		// we do this every frame for now
@@ -270,13 +270,12 @@ void UFullBodyIKSolver::SolveInternal(FIKRigTransformModifier& InOutGlobalTransf
 		}
 
 		// update mid effector info
+		const float LinearMotionStrength = FMath::Max(SolverProperty.LinearMotionStrength, SolverProperty.MinLinearMotionStrength);
+		const float AngularMotionStrength = FMath::Max(SolverProperty.AngularMotionStrength, SolverProperty.MinAngularMotionStrength);
+		const float LinearRange = LinearMotionStrength - SolverProperty.MinLinearMotionStrength;
+		const float AngularRange = AngularMotionStrength - SolverProperty.MinAngularMotionStrength;
 
-		const float LinearMotionStrength = FMath::Max(SolverDef->SolverProperty.LinearMotionStrength, SolverDef->SolverProperty.MinLinearMotionStrength);
-		const float AngularMotionStrength = FMath::Max(SolverDef->SolverProperty.AngularMotionStrength, SolverDef->SolverProperty.MinAngularMotionStrength);
-		const float LinearRange = LinearMotionStrength - SolverDef->SolverProperty.MinLinearMotionStrength;
-		const float AngularRange = AngularMotionStrength - SolverDef->SolverProperty.MinAngularMotionStrength;
-
-		const TArray<FFBIKRigEffector>& Effectors = SolverDef->Effectors;
+		//const TArray<FFBIKRigEffector>& Effectors = Effectors;
 
 		// update end effector info
 		for (int32 EffectorIndex = 0; EffectorIndex < Effectors.Num(); ++EffectorIndex)
@@ -299,11 +298,11 @@ void UFullBodyIKSolver::SolveInternal(FIKRigTransformModifier& InOutGlobalTransf
 
 					const float Pull = FMath::Clamp(CurEffector.Pull, 0.f, 1.f);
 					// we want some impact of Pull, in order for Pull to have some impact, we clamp to some number
-					const float TargetClamp = FMath::Clamp(SolverDef->SolverProperty.DefaultTargetClamp, 0.f, 0.7f);
+					const float TargetClamp = FMath::Clamp(SolverProperty.DefaultTargetClamp, 0.f, 0.7f);
 					const float Scale = TargetClamp + Pull * (1.f - TargetClamp);
 					// Pull set up
-					EffectorTarget->LinearMotionStrength = LinearRange * Scale + SolverDef->SolverProperty.MinLinearMotionStrength;
-					EffectorTarget->AngularMotionStrength = AngularRange * Scale + SolverDef->SolverProperty.MinAngularMotionStrength;
+					EffectorTarget->LinearMotionStrength = LinearRange * Scale + SolverProperty.MinLinearMotionStrength;
+					EffectorTarget->AngularMotionStrength = AngularRange * Scale + SolverProperty.MinAngularMotionStrength;
 					EffectorTarget->ConvergeScale = Scale;
 					EffectorTarget->TargetClampScale = Scale;
 
@@ -315,15 +314,15 @@ void UFullBodyIKSolver::SolveInternal(FIKRigTransformModifier& InOutGlobalTransf
 
 		DebugData.Reset();
 
-		const bool bDebugEnabled = SolverDef->DebugOption.bDrawDebugHierarchy || SolverDef->DebugOption.bDrawDebugEffector || SolverDef->DebugOption.bDrawDebugConstraints;
+		const bool bDebugEnabled = DebugOption.bDrawDebugHierarchy || DebugOption.bDrawDebugEffector || DebugOption.bDrawDebugConstraints;
 
 		// we can't reuse memory until we fix the memory issue on RigVM
 		{
 			IKSolver.SolveJacobianIK(LinkData, EffectorTargets,
-				JacobianIK::FSolverParameter(SolverDef->SolverProperty.Damping, true, false, (SolverDef->SolverProperty.bUseJacobianTranspose) ? EJacobianSolver::JacobianTranspose : EJacobianSolver::JacobianPIDLS),
-				SolverDef->SolverProperty.MaxIterations, SolverDef->SolverProperty.Precision, &DebugData);
+				JacobianIK::FSolverParameter(SolverProperty.Damping, true, false, (SolverProperty.bUseJacobianTranspose) ? EJacobianSolver::JacobianTranspose : EJacobianSolver::JacobianPIDLS),
+				SolverProperty.MaxIterations, SolverProperty.Precision, &DebugData);
 
-			if (SolverDef->MotionProperty.bForceEffectorRotationTarget)
+			if (MotionProperty.bForceEffectorRotationTarget)
 			{
 				// if position is reached, we force rotation target
 				for (int32 EffectorIndex = 0; EffectorIndex < Effectors.Num(); ++EffectorIndex)
@@ -336,13 +335,13 @@ void UFullBodyIKSolver::SolveInternal(FIKRigTransformModifier& InOutGlobalTransf
 						{
 							bool bApplyRotation = true;
 
-							if (SolverDef->MotionProperty.bOnlyApplyWhenReachedToTarget)
+							if (MotionProperty.bOnlyApplyWhenReachedToTarget)
 							{
 								// only do this when position is reached? This will conflict with converge scale
 								const FVector& BonePosition = LinkData[EffectorLinkIndex].GetTransform().GetLocation();
 								const FVector& TargetPosition = EffectorTarget->Position;
 
-								bApplyRotation = (FVector(BonePosition - TargetPosition).SizeSquared() <= SolverDef->SolverProperty.Precision * SolverDef->SolverProperty.Precision);
+								bApplyRotation = (FVector(BonePosition - TargetPosition).SizeSquared() <= SolverProperty.Precision * SolverProperty.Precision);
 							}
 
 							if (bApplyRotation)
@@ -369,10 +368,10 @@ void UFullBodyIKSolver::SolveInternal(FIKRigTransformModifier& InOutGlobalTransf
 				{
 					const TArray<FFBIKLinkData>& LocalLink = DebugData[DebugIndex].LinkData;
 
-					FTransform Offset = SolverDef->DebugOption.DrawWorldOffset;
+					FTransform Offset = DebugOption.DrawWorldOffset;
 					Offset.SetLocation(Offset.GetLocation() * (DebugDataNum - DebugIndex));
 
-					if (SolverDef->DebugOption.bDrawDebugHierarchy)
+					if (DebugOption.bDrawDebugHierarchy)
 					{
 						for (int32 LinkIndex = 0; LinkIndex < LocalLink.Num(); ++LinkIndex)
 						{
@@ -381,25 +380,25 @@ void UFullBodyIKSolver::SolveInternal(FIKRigTransformModifier& InOutGlobalTransf
 							FLinearColor DrawColor = FLinearColor::White;
 
 							float LineThickness = 0.f;
-							if (SolverDef->DebugOption.bColorAngularMotionStrength || SolverDef->DebugOption.bColorLinearMotionStrength)
+							if (DebugOption.bColorAngularMotionStrength || DebugOption.bColorLinearMotionStrength)
 							{
 								DrawColor = FLinearColor::Black;
-								if (SolverDef->DebugOption.bColorAngularMotionStrength)
+								if (DebugOption.bColorAngularMotionStrength)
 								{
-									const float Range = FMath::Max(SolverDef->SolverProperty.AngularMotionStrength - SolverDef->SolverProperty.MinAngularMotionStrength, 0.f);
+									const float Range = FMath::Max(SolverProperty.AngularMotionStrength - SolverProperty.MinAngularMotionStrength, 0.f);
 									if (Range > 0.f)
 									{
-										float CurrentStrength = Data.GetAngularMotionStrength() - SolverDef->SolverProperty.MinAngularMotionStrength;
+										float CurrentStrength = Data.GetAngularMotionStrength() - SolverProperty.MinAngularMotionStrength;
 										float Alpha = FMath::Clamp(CurrentStrength / Range, 0.f, 1.f);
 										DrawColor.R = LineThickness = Alpha;
 									}
 								}
-								else if (SolverDef->DebugOption.bColorLinearMotionStrength)
+								else if (DebugOption.bColorLinearMotionStrength)
 								{
-									const float Range = FMath::Max(SolverDef->SolverProperty.LinearMotionStrength - SolverDef->SolverProperty.MinLinearMotionStrength, 0.f);
+									const float Range = FMath::Max(SolverProperty.LinearMotionStrength - SolverProperty.MinLinearMotionStrength, 0.f);
 									if (Range > 0.f)
 									{
-										float CurrentStrength = Data.GetLinearMotionStrength() - SolverDef->SolverProperty.MinLinearMotionStrength;
+										float CurrentStrength = Data.GetLinearMotionStrength() - SolverProperty.MinLinearMotionStrength;
 										float Alpha = FMath::Clamp(CurrentStrength / Range, 0.f, 1.f);
 										DrawColor.B = LineThickness = Alpha;
 									}
@@ -412,14 +411,14 @@ void UFullBodyIKSolver::SolveInternal(FIKRigTransformModifier& InOutGlobalTransf
 								InOutDrawInterface->DrawLine(Offset, Data.GetPreviousTransform().GetLocation(), ParentData.GetPreviousTransform().GetLocation(), DrawColor, LineThickness);
 							}
 
-							if (SolverDef->DebugOption.bDrawDebugAxes)
+							if (DebugOption.bDrawDebugAxes)
 							{
-								InOutDrawInterface->DrawAxes(Offset, Data.GetPreviousTransform(), SolverDef->DebugOption.DrawSize);
+								InOutDrawInterface->DrawAxes(Offset, Data.GetPreviousTransform(), DebugOption.DrawSize);
 							}
 						}
 					}
 
-					if (SolverDef->DebugOption.bDrawDebugEffector)
+					if (DebugOption.bDrawDebugEffector)
 					{
 						for (auto Iter = EffectorTargets.CreateConstIterator(); Iter; ++Iter)
 						{
@@ -427,11 +426,11 @@ void UFullBodyIKSolver::SolveInternal(FIKRigTransformModifier& InOutGlobalTransf
 							if (EffectorTarget.bPositionEnabled)
 							{
 								// draw effector target locations
-								InOutDrawInterface->DrawBox(Offset, FTransform(EffectorTarget.Position), FLinearColor::Yellow, SolverDef->DebugOption.DrawSize);
+								InOutDrawInterface->DrawBox(Offset, FTransform(EffectorTarget.Position), FLinearColor::Yellow, DebugOption.DrawSize);
 							}
 
 							// draw effector link location
-							InOutDrawInterface->DrawBox(Offset, LocalLink[Iter.Key()].GetPreviousTransform(), FLinearColor::Green, SolverDef->DebugOption.DrawSize);
+							InOutDrawInterface->DrawBox(Offset, LocalLink[Iter.Key()].GetPreviousTransform(), FLinearColor::Green, DebugOption.DrawSize);
 						}
 
 						for (int32 Index = 0; Index < DebugData[DebugIndex].TargetVectorSources.Num(); ++Index)
@@ -560,4 +559,60 @@ void UFullBodyIKSolver::SolveInternal(FIKRigTransformModifier& InOutGlobalTransf
 		const FTransform& LinkTransform = LinkData[LinkIndex].GetTransform();
 		InOutGlobalTransform.SetGlobalTransform(CurrentItem, LinkTransform, true);
 	}
-}
\ No newline at end of file
+}
+
+#if WITH_EDITOR
+
+void UFullBodyIKSolver::UpdateEffectors()
+{
+	// first ensure we add all the effectors
+	for (int32 Index = 0; Index < Effectors.Num(); ++Index)
+	{
+		EnsureToAddEffector(Effectors[Index].Target, TEXT("FBIK_Effector"));
+	}
+
+	// we have more nodes than goals
+	// which means we have something deleted
+	if (Effectors.Num() < EffectorToGoal.Num())
+	{
+		TArray<FIKRigEffector> GoalEffectors;
+		// we have to remove things that don't belong
+		for (auto Iter = EffectorToGoal.CreateIterator(); Iter; ++Iter)
+		{
+			GoalEffectors.Add(Iter.Key());
+		}
+
+		TBitArray<> RemoveFlags(false, GoalEffectors.Num());
+
+		for (int32 Index = 0; Index < Effectors.Num(); ++Index)
+		{
+			int32 Found = GoalEffectors.Find(Effectors[Index].Target);
+
+			if (Found != INDEX_NONE)
+			{
+				RemoveFlags[Found] = true;
+			}
+		}
+
+		for (TConstSetBitIterator<> Iter(RemoveFlags); Iter; ++Iter)
+		{
+			// remove things that don't belong
+			EffectorToGoal.Remove(GoalEffectors[Iter.GetIndex()]);
+		}
+	}
+
+	// trigger a delegate for goal has been updated?
+	OnGoalHasBeenUpdated();
+}
+
+void UFullBodyIKSolver::PostEditChangeChainProperty(struct FPropertyChangedChainEvent& PropertyChangedEvent)
+{
+	if (PropertyChangedEvent.Property->GetFName() == GET_MEMBER_NAME_CHECKED(UFullBodyIKSolver, Effectors))
+	{
+		UpdateEffectors();
+	}
+
+	Super::PostEditChangeChainProperty(PropertyChangedEvent);
+}
+
+#endif// WITH_EDITOR
\ No newline at end of file
diff --git a/Engine/Plugins/Experimental/FullBodyIK/Source/FullBodyIK/Private/FullBodyIKSolver.h b/Engine/Plugins/Experimental/FullBodyIK/Source/FullBodyIK/Private/FullBodyIKSolver.h
--- a/Engine/Plugins/Experimental/FullBodyIK/Source/FullBodyIK/Private/FullBodyIKSolver.h
+++ b/Engine/Plugins/Experimental/FullBodyIK/Source/FullBodyIK/Private/FullBodyIKSolver.h
@@ -9,16 +9,78 @@
 
 #include "IKRigSolver.h"
 #include "FBIKShared.h"
+#include "FBIKConstraintOption.h"
+#include "FBIKDebugOption.h"
 #include "FullBodyIKSolver.generated.h"
 
 struct FFBIKLinkData;
 
+USTRUCT()
+struct FFBIKRigEffector
+{
+	GENERATED_BODY()
+
+	/**
+	* The last item in the chain to solve - the effector
+	*/
+	UPROPERTY(EditAnywhere, Category = FFBIKRigEffector)
+	FIKRigEffector Target;
+
+	UPROPERTY(EditAnywhere, Category = FFBIKRigEffector)
+	int32  PositionDepth = 1000;
+
+	UPROPERTY(EditAnywhere, Category = FFBIKRigEffector)
+	int32  RotationDepth = 1000;
+
+	/*
+	 * Clamps the total length to target by this scale for each iteration
+	 * This helps to stabilize solver to reduce singularity by avoiding to try to reach target too far.
+	 */
+	UPROPERTY(EditAnywhere, Category = FFBIKRigEffector)
+	float Pull = 0.f;
+
+	FFBIKRigEffector()
+	{
+	}
+};
+
  // run time for UFullBodyIKSolverDefinition
 UCLASS()
 class FULLBODYIK_API UFullBodyIKSolver : public UIKRigSolver
 {
 	GENERATED_BODY()
 
+public:
+	UPROPERTY(EditAnywhere, Category = Effectors)
+	FName Root;
+
+	UPROPERTY(EditAnywhere, Category = Effectors)
+	TArray<FFBIKRigEffector> Effectors;
+	//	UPROPERTY(EditAnywhere, Category = Constraint)
+	//	TArray<FFBIKConstraintOption> Constraints;
+
+	UPROPERTY(EditAnywhere, Category = Solver)
+	FSolverInput	SolverProperty;
+
+	UPROPERTY(EditAnywhere, Category = Motion)
+	FMotionProcessInput MotionProperty;
+
+	UPROPERTY(EditAnywhere, Category = Solver)
+	FFBIKDebugOption DebugOption;
+
+	UFullBodyIKSolver();
+
+private:
+
+#if WITH_EDITOR
+	virtual void UpdateEffectors() override;
+
+	// UObject interface
+	virtual void PostEditChangeChainProperty(struct FPropertyChangedChainEvent& PropertyChangedEvent) override;
+	// UObject interface
+#endif // WITH_EDITOR
+
+
 	/** list of Link Data for solvers - joints */
 	TArray<FFBIKLinkData> LinkData;
 	/** Effector Targets - search key is LinkData Index */
@@ -40,7 +102,7 @@ protected:
 	virtual void InitInternal(const FIKRigTransformModifier& InGlobalTransform) override;
 	virtual void SolveInternal(FIKRigTransformModifier& InOutGlobalTransform, FControlRigDrawInterface* InOutDrawInterface) override;
 	virtual bool IsSolverActive() const override;
-};
-
-
 
+public:
+	static const FName EffectorTargetPrefix;
+};
\ No newline at end of file
diff --git a/Engine/Plugins/Experimental/FullBodyIK/Source/FullBodyIK/Private/FullBodyIKSolverDefinition.cpp b/Engine/Plugins/Experimental/FullBodyIK/Source/FullBodyIK/Private/FullBodyIKSolverDefinition.cpp
deleted file mode 100644
diff --git a/Engine/Plugins/Experimental/FullBodyIK/Source/FullBodyIK/Private/FullBodyIKSolverDefinition.h b/Engine/Plugins/Experimental/FullBodyIK/Source/FullBodyIK/Private/FullBodyIKSolverDefinition.h
deleted file mode 100644
