Fix goal/effector connection method in IKRig plugin. General refactoring.

diff --git a/Engine/Plugins/Animation/IKRig/Source/IKRig/Private/AnimNodes/AnimNode_IKRig.cpp b/Engine/Plugins/Animation/IKRig/Source/IKRig/Private/AnimNodes/AnimNode_IKRig.cpp
--- a/Engine/Plugins/Animation/IKRig/Source/IKRig/Private/AnimNodes/AnimNode_IKRig.cpp
+++ b/Engine/Plugins/Animation/IKRig/Source/IKRig/Private/AnimNodes/AnimNode_IKRig.cpp
@@ -41,7 +41,7 @@ void FAnimNode_IKRig::Evaluate_AnyThread(FPoseContext& Output)
 	}
 		
 	FCompactPose& OutPose = SourcePose.Pose;
-	FIKRigTransforms& Transforms = RigProcessor->GetTransforms();
+	FIKRigTransforms& Transforms = RigProcessor->GetCurrentGlobalTransforms();
 
 	// copy input pose
 	for (FCompactPoseBoneIndex CPIndex : OutPose.ForEachBoneIndex())
@@ -62,10 +62,10 @@ void FAnimNode_IKRig::Evaluate_AnyThread(FPoseContext& Output)
 	{
 		if (ensure(GoalTransforms.IsValidIndex(GoalIndex)))
 		{
-			RigProcessor->Goals.SetGoalTransform(
+			RigProcessor->SetGoalTransform(
 				GoalNames[GoalIndex], 
 				GoalTransforms[GoalIndex].GetLocation(),
-				GoalTransforms[GoalIndex].GetRotation().Rotator());
+				GoalTransforms[GoalIndex].GetRotation());
 		}
 	}
 
@@ -106,7 +106,7 @@ void FAnimNode_IKRig::Initialize_AnyThread(const FAnimationInitializeContext& Co
 	{
 		RigProcessor->Initialize(RigDefinitionAsset);
 		GoalNames.Reset();
-		RigProcessor->Goals.GetNames(GoalNames);
+		RigProcessor->GetGoalNames(GoalNames);
 	}
 }
 
@@ -126,7 +126,9 @@ bool FAnimNode_IKRig::RebuildGoalList()
 {
 	if (RigDefinitionAsset)
 	{
-		const int32 GoalNum = RigDefinitionAsset->GetGoals().Num();
+		TArray<FName> SolverGoals;
+		RigDefinitionAsset->GetGoalNamesFromSolvers(SolverGoals);
+		const int32 GoalNum = SolverGoals.Num();
 		if (GoalTransforms.Num() != GoalNum)
 		{
 			GoalTransforms.SetNum(GoalNum);
@@ -141,9 +143,8 @@ FName FAnimNode_IKRig::GetGoalName(int32 Index) const
 {
 	if (RigDefinitionAsset)
 	{
-		// TODO: fix  this vs RigProcessor goals. I think RigProcessor
-		// only issue with RigProcessor is you don't know until it compiles
-		TArray<FName> Names = RigDefinitionAsset->GetGoalNames();
+		TArray<FName> Names;
+		RigDefinitionAsset->GetGoalNamesFromSolvers(Names);
 		if (Names.IsValidIndex(Index))
 		{
 			return Names[Index];
diff --git a/Engine/Plugins/Animation/IKRig/Source/IKRig/Private/IKRigDefinition.cpp b/Engine/Plugins/Animation/IKRig/Source/IKRig/Private/IKRigDefinition.cpp
--- a/Engine/Plugins/Animation/IKRig/Source/IKRig/Private/IKRigDefinition.cpp
+++ b/Engine/Plugins/Animation/IKRig/Source/IKRig/Private/IKRigDefinition.cpp
@@ -3,78 +3,27 @@
 #include "IKRigDefinition.h"
 #include "IKRigSolver.h"
 
-UIKRigDefinition::UIKRigDefinition()
-{
-}
 
-#if WITH_EDITOR
-void UIKRigDefinition::UpdateGoal()
+void UIKRigDefinition::GetGoalNamesFromSolvers(TArray<FName>& OutGoalNames) const
 {
-	// collect goals from the solver definitions
-	// each solver definition collects their tasks to goals
-	// but the IKRigDefinition will maintain the default value and so on
-	Sanitize();
-
-	TArray<FName> ListOfGoals;
 	for (UIKRigSolver* Solver : Solvers)
 	{
-		if (Solver)
-		{
-			Solver->AppendGoalNamesToArray(ListOfGoals);
-		}
-	}
-
-	// we want to update goal list properly
-
-	// first if it's not used any more, remove
-	TArray<FName> GoalsToRemove;
+		TSet<FName> GoalNames;
+		Solver->CollectGoalNames(GoalNames);
 
-	// first we go through IKGoals and see if we should remove unused ones
-	for (auto Iter = Goals.CreateConstIterator(); Iter; ++Iter)
-	{
-		// if it's not used any more
-		if (!ListOfGoals.Contains(Iter->Key))
+		// not using a TSet here because user code relies on indices
+		for (FName Name : GoalNames)
 		{
-			GoalsToRemove.Add(Iter->Key);
-		}
-	}
-
-	for (const FName& Goal : GoalsToRemove)
-	{
-		Goals.Remove(Goal);
-	}
-
-	// now add new ones if it doesn't exist
-	for (int32 Index=0; Index<ListOfGoals.Num(); ++Index)
-	{
-		const FName& Goal = ListOfGoals[Index];
-		if (!Goals.Contains(Goal))
-		{
-			// if we don't have it, add to new one
-			FIKRigGoal NewGoal(Goal);
-			Goals.Add(Goal, NewGoal);
+			if (!OutGoalNames.Contains(Name))
+			{
+				OutGoalNames.Add(Name);
+			}
 		}
 	}
-
-	// this doesn't keep the old list alive
-	// unsure if that's what we want yet. Sometimes you may lose
-	// your used variable, but that's something we can find it out 
-	// from UX
 }
 
-void UIKRigDefinition::Sanitize()
-{
-	// sanitize to clean up
-	Solvers.RemoveAll([this](const UIKRigSolver* Solver) { return Solver == nullptr; });
-}
 
-/////////////////////////////////////////////////////////
-//// Hierarchy modifiers
-// why this is here? 
-// delegate and user instancing 
-// can this be in controller?
-// 
-/////////////////////////////////////////////////////////
+#if WITH_EDITOR
 
 // add a bone. Return false if it already has conflict name or parent is not found
 bool UIKRigDefinition::AddBone(const FName& InName, const FName& InParent, const FTransform& InGlobalTransform)
@@ -91,7 +40,6 @@ bool UIKRigDefinition::AddBone(const FName& InName, const FName& InParent, const
 			check (Hierarchy.Bones.Num() == RefPoseTransforms.Num());
 
 			Hierarchy.RebuildCacheData();
-			BoneAddedDelegate.Broadcast(InName);
 			return true;
 		}
 	}
@@ -99,127 +47,6 @@ bool UIKRigDefinition::AddBone(const FName& InName, const FName& InParent, const
 	return false;
 }
 
-// remove a bone. Returns false if it is not found
-// if it has a children, it will remove all the children with it
-bool UIKRigDefinition::RemoveBone(const FName& InName)
-{
-	struct FRemoveChild
-	{
-		static void RemoveChild_Recursively(const FIKRigHierarchy& InHierarchy, TArray<FIKRigBone>& InOutBones, TArray<FTransform>& InOutTransform, const FName& InName)
-		{
-			// ensure I have this bone
-			int32 Index = InHierarchy.FindIndexFromBoneArray(InName);
-			if (Index != INDEX_NONE)
-			{
-				// find children
-				TArray<int32> ChildIndices = InHierarchy.FindIndicesByParentName(InName);
-				// iterate back from the forward
-				for (int32 ChildIndex = ChildIndices.Num() - 1; ChildIndex >= 0; --ChildIndex)
-				{
-					// remove children's first, since children is always behind of parent, this should work
-					RemoveChild_Recursively(InHierarchy, InOutBones, InOutTransform, InOutBones[ChildIndices[ChildIndex]].Name);
-					// do not shrink yet and remove all children first
-					// since children is behind of parent, this works
-				}
-
-				// now remove parent (this)
-				InOutBones.RemoveAt(Index, 1, false);
-				InOutTransform.RemoveAt(Index, 1, false);
-			}
-		}
-	};
-
-	// ensure I have this bone
-	int32 Index = Hierarchy.FindIndexFromBoneArray(InName);
-	if (Index != INDEX_NONE)
-	{
-		// remove all children
-		FRemoveChild::RemoveChild_Recursively(Hierarchy, Hierarchy.Bones, RefPoseTransforms, InName);
-
-		// shrink now
-		Hierarchy.Bones.Shrink();
-		RefPoseTransforms.Shrink();
-		Hierarchy.RebuildCacheData();
-		BoneRemovedDelegate.Broadcast(InName);
-		return true;
-	}
-
-	return false;
-}
-
-// rename just change the name
-// return false if not found
-bool UIKRigDefinition::RenameBone(const FName& InOldName, const FName& InNewName)
-{
-	// ensure we have old name and do not have new name
-	int32 Index = Hierarchy.FindIndexFromBoneArray(InOldName);
-	if (Index != INDEX_NONE)
-	{
-		// make sure nobody is using NewName
-		if (Hierarchy.FindIndexFromBoneArray(InNewName) == INDEX_NONE)
-		{
-			TArray<int32> ChildIndices = Hierarchy.FindIndicesByParentName(InOldName);
-
-			for (int32 ChildIndex = 0; ChildIndex < ChildIndices.Num(); ++ChildIndex)
-			{
-				Hierarchy.Bones[ChildIndices[ChildIndex]].ParentName = InNewName;
-			}
-
-			// now change me
-			Hierarchy.Bones[Index].Name = InNewName;
-			Hierarchy.RebuildCacheData();
-			BoneRenamedDelegate.Broadcast(InOldName, InNewName);
-			return true;
-		}
-	}
-
-	return false;
-}
-
-// reparent to InNewParent
-//	if not found - either InName or InParent
-//	Or if InNewParent is invalid (i.e. it's a child of InName)
-bool UIKRigDefinition::ReparentBone(const FName& InName, const FName& InNewParent)
-{
-	int32 Index = Hierarchy.FindIndexFromBoneArray(InName);
-	if (Index != INDEX_NONE)
-	{
-		int32 NewParentIndex = INDEX_NONE;
-		if (InNewParent != NAME_None)
-		{
-			NewParentIndex = Hierarchy.FindIndexFromBoneArray(InNewParent);
-			if (NewParentIndex == INDEX_NONE)
-			{
-				return false;
-			}
-
-			// we have to ensure we don't have cycling situation
-			int32 CurIndex = NewParentIndex;
-			while (CurIndex != INDEX_NONE)
-			{
-				// i hit myself
-				if (CurIndex == Index)
-				{
-					// we're cycling
-					return false;
-				}
-
-				// find parent
-				CurIndex = Hierarchy.FindIndexFromBoneArray(Hierarchy.Bones[CurIndex].ParentName);
-			}
-		}
-
-		// we don't have a cycling, now allow set it
-		Hierarchy.Bones[Index].ParentName = InNewParent;
-		EnsureSortedCorrectly(true);
-		BoneReparentedDelegate.Broadcast(InName, InNewParent);
-		return true;
-	}
-
-	return false;
-}
-
-
 // ensure it's sorted from parent to children
 void UIKRigDefinition::EnsureSortedCorrectly(bool bReSortIfNeeded)
 {
@@ -352,21 +179,6 @@ void UIKRigDefinition::ResetHierarchy()
 	RefPoseTransforms.Reset();
 }
 
-void UIKRigDefinition::EnsureCreateUniqueGoalName(FName& InOutGoal) const
-{
-	int32 Index = 1;
-
-	FString GoalNameString = InOutGoal.ToString();
-
-	// while contains
-	while(Goals.Find(FName(*GoalNameString)))
-	{
-		GoalNameString = FString::Format(TEXT("{0}_{1}"), {InOutGoal.ToString(), Index++});
-	}
-
-	InOutGoal = FName(*GoalNameString);
-}
-
 #endif // WITH_EDITOR
 
 void UIKRigDefinition::PostLoad()
@@ -374,13 +186,4 @@ void UIKRigDefinition::PostLoad()
 	Super::PostLoad();
 
 	Hierarchy.RebuildCacheData();
-}
-
-void UIKRigDefinition::BeginDestroy()
-{
-#if WITH_EDITOR
-	IKRigDefinitionBeginDestroy.Broadcast(this);
-#endif // #if WITH_EDITOR
-
-	Super::BeginDestroy();
 }
\ No newline at end of file
diff --git a/Engine/Plugins/Animation/IKRig/Source/IKRig/Private/IKRigProcessor.cpp b/Engine/Plugins/Animation/IKRig/Source/IKRig/Private/IKRigProcessor.cpp
--- a/Engine/Plugins/Animation/IKRig/Source/IKRig/Private/IKRigProcessor.cpp
+++ b/Engine/Plugins/Animation/IKRig/Source/IKRig/Private/IKRigProcessor.cpp
@@ -16,11 +16,13 @@ void UIKRigProcessor::Initialize(UIKRigDefinition* InRigDefinition)
 		return;
 	}
 	
-	RefPoseTransforms = RigDefinition->GetReferencePose();
-	Transforms = FIKRigTransforms(&RigDefinition->GetHierarchy());
-	Transforms.SetAllGlobalTransforms(RefPoseTransforms);
-	
-	Goals.SetAllGoals(RigDefinition->GetGoals());
+	RefPoseGlobalTransforms = RigDefinition->GetReferencePose();
+	GlobalBoneTransforms = FIKRigTransforms(&RigDefinition->GetHierarchy());
+	GlobalBoneTransforms.SetAllGlobalTransforms(RefPoseGlobalTransforms);
+
+	TArray<FName> GoalNames;
+	RigDefinition->GetGoalNamesFromSolvers(GoalNames);
+	Goals.InitializeGoalsFromNames(GoalNames);
 
 	TArray<UIKRigSolver*> RigSolvers = RigDefinition->GetSolvers();
 	const int32 NumSolvers = RigSolvers.Num();
@@ -28,7 +30,7 @@ void UIKRigProcessor::Initialize(UIKRigDefinition* InRigDefinition)
 	for (int32 Index = 0; Index < NumSolvers; ++Index)
 	{
 		UIKRigSolver* Solver = DuplicateObject(RigSolvers[Index], this);
-		Solver->Init(Transforms);
+		Solver->Init(GlobalBoneTransforms);
 		Solvers.Add(Solver);
 	}
 
@@ -37,7 +39,7 @@ void UIKRigProcessor::Initialize(UIKRigDefinition* InRigDefinition)
 
 void UIKRigProcessor::Solve()
 {
-	if (!bInitialized)
+	if (!ensure(bInitialized))
 	{
 		return;
 	}
@@ -48,26 +50,44 @@ void UIKRigProcessor::Solve()
 
 	for (int32 Index = 0; Index < Solvers.Num(); ++Index)
 	{
-		Solvers[Index]->Solve(Transforms, Goals, &DrawInterface);
+		Solvers[Index]->Solve(GlobalBoneTransforms, Goals, &DrawInterface);
 	}
 }
 
-FIKRigTransforms& UIKRigProcessor::GetTransforms() 
+FIKRigTransforms& UIKRigProcessor::GetCurrentGlobalTransforms() 
 {
-	return Transforms;
+	check(bInitialized);
+	return GlobalBoneTransforms;
 }
 
-const FIKRigHierarchy* UIKRigProcessor::GetHierarchy() const
+void UIKRigProcessor::SetGoalTransform(
+	const FName& GoalName,
+	const FVector& Position,
+	const FQuat& Rotation)
 {
-	if (RigDefinition)
-	{
-		return &RigDefinition->GetHierarchy();
-	}
+	check(bInitialized);
+	Goals.SetGoalTransform(GoalName, Position, Rotation);
+}
+
+void UIKRigProcessor::GetGoalNames(TArray<FName>& OutGoalNames) const
+{
+	check(bInitialized);
+	return Goals.GetNames(OutGoalNames);
+}
 
-	return nullptr;
+int UIKRigProcessor::GetNumGoals() const
+{
+	return Goals.GetNumGoals();
+}
+
+const FIKRigHierarchy* UIKRigProcessor::GetHierarchy() const
+{
+	check(bInitialized);
+	return ensure(RigDefinition) ? &RigDefinition->GetHierarchy(): nullptr;
 }
 
 void UIKRigProcessor::ResetToRefPose()
 {
-	Transforms.SetAllGlobalTransforms(RefPoseTransforms);
+	check(bInitialized);
+	GlobalBoneTransforms.SetAllGlobalTransforms(RefPoseGlobalTransforms);
 }
\ No newline at end of file
diff --git a/Engine/Plugins/Animation/IKRig/Source/IKRig/Private/IKRigSolver.cpp b/Engine/Plugins/Animation/IKRig/Source/IKRig/Private/IKRigSolver.cpp
--- a/Engine/Plugins/Animation/IKRig/Source/IKRig/Private/IKRigSolver.cpp
+++ b/Engine/Plugins/Animation/IKRig/Source/IKRig/Private/IKRigSolver.cpp
@@ -1,49 +1,13 @@
 // Copyright Epic Games, Inc. All Rights Reserved.
 
-
 #include "IKRigSolver.h"
 #include "IKRigDataTypes.h"
-#include "IKRigDefinition.h"
-
-
-void UIKRigSolver::SolveInternal(
-	FIKRigTransforms& InOutGlobalTransform,
-	const FIKRigGoalContainer& Goals,
-	FControlRigDrawInterface* InOutDrawInterface)
-{
-	if (IsSolverActive())
-	{
-		Solve(InOutGlobalTransform, Goals, InOutDrawInterface);
-	}
-}
 
-bool UIKRigSolver::IsSolverActive() const
-{
-	return bEnabled;
-}
-
-bool UIKRigSolver::GetGoalForEffector(
-	const FIKRigEffector& InEffector,
-	const FIKRigGoalContainer& Goals,
-	FIKRigGoal& OutGoal) const
-{
-	const FName* GoalName = EffectorToGoalName.Find(InEffector);
-	if (!GoalName)
-	{
-		return false;
-	}
-
-	return Goals.GetGoalByName(*GoalName, OutGoal);
-}
-
-void UIKRigSolver::AppendGoalNamesToArray(TArray<FName>& OutGoals)
-{
-	TArray<FName> LocalGoals;
-	EffectorToGoalName.GenerateValueArray(LocalGoals);
-	OutGoals += LocalGoals;
-}
 
+/* TODO KIARAN: do we even want to auto rename goals?
+ * I think the risk is high of being frustrating for user
 #if WITH_EDITOR
+
 void UIKRigSolver::RenameGoal(const FName& OldName, const FName& NewName)
 {
 	for (auto Iter = EffectorToGoalName.CreateIterator(); Iter; ++Iter)
@@ -55,63 +19,5 @@ void UIKRigSolver::RenameGoal(const FName& OldName, const FName& NewName)
 	}
 }
 
-void UIKRigSolver::EnsureUniqueGoalName(FName& InOutUniqueGoalName) const
-{
-	// call delegate 
-	UIKRigDefinition* IKRigDef = CastChecked<UIKRigDefinition>(GetOuter());
-	IKRigDef->EnsureCreateUniqueGoalName(InOutUniqueGoalName);
-}
-
-FName UIKRigSolver::CreateUniqueGoalName(const TCHAR* Suffix) const
-{
-	if (Suffix)
-	{
-		FString NewGoalStr = FString("") + TEXT("NewGoal_") + Suffix;
-		// replace any whitespace with _
-		NewGoalStr = NewGoalStr.Replace(TEXT(" \t"), TEXT("_"));
-		FName NewGoalName = FName(*NewGoalStr, FNAME_Add);
-		EnsureUniqueGoalName(NewGoalName);
-
-		return NewGoalName;
-	}
-
-	return NAME_None;
-}
-
-void UIKRigSolver::OnGoalHasBeenUpdated()
-{
-	GoalNeedsUpdateDelegate.Broadcast();
-}
-
-void UIKRigSolver::EnsureToAddEffector(const FIKRigEffector& InEffector, const FString& InPrefix)
-{
-	FName* GoalName = EffectorToGoalName.Find(InEffector);
-	if (!GoalName)
-	{
-		EffectorToGoalName.Add(InEffector) = CreateUniqueGoalName(*InPrefix);
-	}
-}
-
-void UIKRigSolver::EnsureToRemoveEffector(const FIKRigEffector& InEffector)
-{
-	EffectorToGoalName.Remove(InEffector);
-}
-
-#endif // WITH_EDITOR
-
-void UIKRigSolver::PostLoad()
-{
-	Super::PostLoad();
-
-#if WITH_EDITOR
-	UpdateEffectors();
 #endif // WITH_EDITOR
-}
-
-void UIKRigSolver::Serialize(FArchive& Ar)
-{
-	Super::Serialize(Ar);
-
-	// serialize these manually because that's the custom types
-	Ar << EffectorToGoalName;
-}
+*/
\ No newline at end of file
diff --git a/Engine/Plugins/Animation/IKRig/Source/IKRig/Private/Solvers/ConstraintSolver.cpp b/Engine/Plugins/Animation/IKRig/Source/IKRig/Private/Solvers/ConstraintSolver.cpp
--- a/Engine/Plugins/Animation/IKRig/Source/IKRig/Private/Solvers/ConstraintSolver.cpp
+++ b/Engine/Plugins/Animation/IKRig/Source/IKRig/Private/Solvers/ConstraintSolver.cpp
@@ -3,42 +3,9 @@
 #include "Solvers/ConstraintSolver.h"
 #include "IKRigBoneSetting.h"
 
-/*
-// register/unregister query function
-void UIKRigConstraintSolver::RegisterQueryConstraintHandler(const FIKRigQueryConstraint& InQueryConstraintHandler)
-{
-
-}
-
-void UIKRigConstraintSolver::UnregisterQueryConstraintHandler()
-{
-
-}
-*/
-
 void UIKRigConstraintSolver::Init(const FIKRigTransforms& InGlobalTransform)
 {
-	/*
-	if (ConstraintDefinition)
-	{
-		// we duplicate and create a copy of the objects, so that we can mutate
-		ConstraintProfiles = ConstraintDefinition->ConstraintProfiles;
-		for (auto Iter = ConstraintProfiles.CreateIterator(); Iter; ++Iter)
-		{
-			FIKRigConstraintProfile& CurrentProfile = Iter.Value();
-			for (auto InnerIter = CurrentProfile.Constraints.CreateIterator(); InnerIter; ++InnerIter)
-			{
-				UIKRigConstraint*& Template = InnerIter.Value();
-				UIKRigConstraint* CopyConstraint = NewObject<UIKRigConstraint>(this, Template->GetClass(), TEXT("IKRigConstraint"), RF_NoFlags, Template);
-				CopyConstraint->Setup(InGlobalTransform);
-				// replace value
-				Template = CopyConstraint;
-			}
-		}
-
-		// go through all UIKRigConstraint and create copy
-		ActiveProfile = UIKRigConstraintDefinition::DefaultProfileName;
-	}*/
+	// TODO get constraints that this solver can deal with
 }
 
 void UIKRigConstraintSolver::Solve(
@@ -63,22 +30,7 @@ void UIKRigConstraintSolver::Solve(
 	}*/
 }
 
-/*
-void UIKRigConstraintSolver::SolverConstraints(const TArray<FName>& ConstraintsList, FIKRigTransformModifier& InOutGlobalTransform, FControlRigDrawInterface* InOutDrawInterface)
+void UIKRigConstraintSolver::CollectGoalNames(TSet<FName>& OutGoals) const
 {
-	FIKRigConstraintProfile* Current = ConstraintProfiles.Find(ActiveProfile);
 
-	if (Current)
-	{
-		for (const FName& ConstraintName : ConstraintsList)
-		{
-			UIKRigConstraint** Constraint = Current->Constraints.Find(ConstraintName);
-			if (Constraint)
-			{
-				// @todo: later add inoutdrawinterface
-				(*Constraint)->Apply(InOutGlobalTransform, InOutDrawInterface);
-			}
-		}
-	}
-}
-*/
\ No newline at end of file
+}
\ No newline at end of file
diff --git a/Engine/Plugins/Animation/IKRig/Source/IKRig/Private/Solvers/TransformSolver.cpp b/Engine/Plugins/Animation/IKRig/Source/IKRig/Private/Solvers/TransformSolver.cpp
--- a/Engine/Plugins/Animation/IKRig/Source/IKRig/Private/Solvers/TransformSolver.cpp
+++ b/Engine/Plugins/Animation/IKRig/Source/IKRig/Private/Solvers/TransformSolver.cpp
@@ -7,8 +7,8 @@
 
 
 UTransformSolver::UTransformSolver()
-	: TransformTargetName(TEXT("TransformTarget"))
 {
+	Effector.Goal = "DefaultGoal";
 }
 
 void UTransformSolver::Init(const FIKRigTransforms& InGlobalTransform)
@@ -16,11 +16,6 @@ void UTransformSolver::Init(const FIKRigTransforms& InGlobalTransform)
 	
 }
 
-bool UTransformSolver::IsSolverActive() const
-{
-	return Super::IsSolverActive() && (bEnablePosition || bEnableRotation);
-}
-
 void UTransformSolver::Solve(
 	FIKRigTransforms& InOutGlobalTransform,
 	const FIKRigGoalContainer& Goals,
@@ -46,46 +41,13 @@ void UTransformSolver::Solve(
 	}
 	if (bEnableRotation)
 	{
-		CurrentTransform.SetRotation(Goal.Rotation.Quaternion());
+		CurrentTransform.SetRotation(Goal.Rotation);
 	}
 
 	InOutGlobalTransform.SetGlobalTransform(Index, CurrentTransform, true);
 }
 
-#if WITH_EDITOR
-
-void UTransformSolver::UpdateEffectors()
-{
-	// KIARAN - this doesn't make any sense, why would we remove an effector just because
-	// the solver settings have been toggled off?
-
-	// update tasks - you want to keep old name though if changed
-	bool bActiveTask = bEnablePosition || bEnableRotation;
-
-	// not found, we add new
-	if (bActiveTask)
-	{
-		EnsureToAddEffector(Effector, *TransformTargetName);
-	}
-	// if the opposite, we have to remove
-	else if (!bActiveTask)
-	{
-		EnsureToRemoveEffector(Effector);
-	}
-
-	// trigger a delegate for goal has been updated?
-	OnGoalHasBeenUpdated();
-}
-
-void UTransformSolver::PostEditChangeChainProperty(struct FPropertyChangedChainEvent& PropertyChangedEvent)
+void UTransformSolver::CollectGoalNames(TSet<FName>& OutGoals) const
 {
-	if (PropertyChangedEvent.Property->GetFName() == GET_MEMBER_NAME_CHECKED(UTransformSolver, bEnablePosition) ||
-		PropertyChangedEvent.Property->GetFName() == GET_MEMBER_NAME_CHECKED(UTransformSolver, bEnableRotation))
-	{
-		UpdateEffectors();
-	}
-
-	Super::PostEditChangeChainProperty(PropertyChangedEvent);
+	OutGoals.Add(Effector.Goal);
 }
-
-#endif// WITH_EDITOR
\ No newline at end of file
diff --git a/Engine/Plugins/Animation/IKRig/Source/IKRig/Public/BoneSettings/RotationLimits.h b/Engine/Plugins/Animation/IKRig/Source/IKRig/Public/BoneSettings/RotationLimits.h
--- a/Engine/Plugins/Animation/IKRig/Source/IKRig/Public/BoneSettings/RotationLimits.h
+++ b/Engine/Plugins/Animation/IKRig/Source/IKRig/Public/BoneSettings/RotationLimits.h
@@ -13,7 +13,7 @@ enum class ERotationLimitMode : uint8
 	Locked	UMETA(DisplayName = "Locked"),
 };
 
-UCLASS(config = Engine, hidecategories = UObject, BlueprintType)
+UCLASS(EditInlineNew, config = Engine, hidecategories = UObject, BlueprintType)
 class IKRIG_API URotationLimits : public UIKRigBoneSetting
 {
 	GENERATED_BODY()
diff --git a/Engine/Plugins/Animation/IKRig/Source/IKRig/Public/IKRigBoneSetting.h b/Engine/Plugins/Animation/IKRig/Source/IKRig/Public/IKRigBoneSetting.h
--- a/Engine/Plugins/Animation/IKRig/Source/IKRig/Public/IKRigBoneSetting.h
+++ b/Engine/Plugins/Animation/IKRig/Source/IKRig/Public/IKRigBoneSetting.h
@@ -7,7 +7,7 @@
 #include "UObject/Object.h"
 #include "IKRigBoneSetting.generated.h"
 
-UCLASS(config = Engine, hidecategories = UObject, DontCollapseCategories, BlueprintType)
+UCLASS(abstract, config = Engine, BlueprintType)
 class IKRIG_API UIKRigBoneSetting : public UObject
 {
 	GENERATED_BODY()
diff --git a/Engine/Plugins/Animation/IKRig/Source/IKRig/Public/IKRigDataTypes.h b/Engine/Plugins/Animation/IKRig/Source/IKRig/Public/IKRigDataTypes.h
--- a/Engine/Plugins/Animation/IKRig/Source/IKRig/Public/IKRigDataTypes.h
+++ b/Engine/Plugins/Animation/IKRig/Source/IKRig/Public/IKRigDataTypes.h
@@ -11,13 +11,12 @@
 
 struct FIKRigHierarchy;
 
-
 USTRUCT()
 struct IKRIG_API FIKRigGoal
 {
 	GENERATED_BODY()
 
-	FIKRigGoal()
+		FIKRigGoal()
 		: Position(ForceInitToZero),
 		Rotation(ForceInitToZero)
 	{
@@ -31,64 +30,38 @@ struct IKRIG_API FIKRigGoal
 	}
 
 	UPROPERTY(VisibleAnywhere, Category = FIKRigGoal)
-	FName Name;
+		FName Name;
 
 	UPROPERTY(EditAnywhere, Category = FIKRigGoal)
-	FVector Position;
+		FVector Position;
 
 	UPROPERTY(EditAnywhere, Category = FIKRigGoal)
-	FRotator Rotation;
+		FQuat Rotation;
 };
 
-USTRUCT()
-struct IKRIG_API FIKRigEffector
-{
-	GENERATED_BODY()
-
-	UPROPERTY()
-	FGuid Guid;
-
-	UPROPERTY(EditAnywhere, Category = FIKRigEffector)
-	FName Bone;
-
-	FIKRigEffector()
-		: Guid(FGuid::NewGuid())
-	{
-	}
-
-	friend FArchive& operator<<(FArchive& Ar, FIKRigEffector& Effector)
-	{
-		return Ar << Effector.Guid;
-	}
-
-	FORCEINLINE bool operator==(const FIKRigEffector& Rhs) const
-	{
-		return Guid == Rhs.Guid;
-	}
-};
+inline uint32 GetTypeHash(FIKRigGoal ObjectRef) { return GetTypeHash(ObjectRef.Name); }
 
 USTRUCT()
 struct IKRIG_API FIKRigGoalContainer
 {
 	GENERATED_BODY()
 
-private:
-
-	TMap<FName, FIKRigGoal> Goals;
-
-public:
-
-	void SetAllGoals(const TMap<FName, FIKRigGoal> &InGoals)
+		void InitializeGoalsFromNames(const TArray<FName>& InGoalNames)
 	{
-		Goals = InGoals;
+		Goals.Reserve(InGoalNames.Num());
+		for (const FName& Name : InGoalNames)
+		{
+			Goals.Emplace(Name, Name);
+		}
 	}
 
 	void SetGoalTransform(
 		const FName& GoalName,
 		const FVector& InPosition,
-		const FRotator& InRotation)
+		const FQuat& InRotation)
 	{
-		if (FIKRigGoal* Goal = Goals.Find(GoalName))
+		FIKRigGoal* Goal = Goals.Find(GoalName);
+		if (Goal)
 		{
 			Goal->Position = InPosition;
 			Goal->Rotation = InRotation;
@@ -111,22 +84,22 @@ public:
 	{
 		Goals.GenerateKeyArray(OutNames);
 	}
-};
 
-template <typename ValueType>
-struct TIKRigEffectorMapKeyFuncs : public TDefaultMapKeyFuncs<const FIKRigEffector, ValueType, false>
-{
-	static FORCEINLINE FIKRigEffector	GetSetKey(TPair<FIKRigEffector, ValueType> const& Element) { return Element.Key; }
-	static FORCEINLINE uint32			GetKeyHash(FIKRigEffector const& Key) { return GetTypeHash(Key.Guid); }
-	static FORCEINLINE bool				Matches(FIKRigEffector const& A, FIKRigEffector const& B) { return (A.Guid == B.Guid); }
-};
+	FORCEINLINE int GetNumGoals() const
+	{
+		return Goals.Num();
+	}
 
-template <typename ValueType>
-using TIKRigEffectorMap = TMap<FIKRigEffector, ValueType, FDefaultSetAllocator, TIKRigEffectorMapKeyFuncs<ValueType>>;
+private:
 
+	TMap<FName, FIKRigGoal> Goals;
+};
 
+USTRUCT()
 struct IKRIG_API FIKRigTransforms
 {
+	GENERATED_BODY()
+	
 	FIKRigTransforms()
 		: Hierarchy(nullptr)
 	{
diff --git a/Engine/Plugins/Animation/IKRig/Source/IKRig/Public/IKRigDefinition.h b/Engine/Plugins/Animation/IKRig/Source/IKRig/Public/IKRigDefinition.h
--- a/Engine/Plugins/Animation/IKRig/Source/IKRig/Public/IKRigDefinition.h
+++ b/Engine/Plugins/Animation/IKRig/Source/IKRig/Public/IKRigDefinition.h
@@ -6,8 +6,9 @@
 #include "UObject/ObjectMacros.h"
 #include "UObject/Object.h"
 #include "IKRigBoneSetting.h"
-#include "IKRigDataTypes.h"
 #include "IKRigHierarchy.h"
+#include "IKRigDataTypes.h"
+
 #include "IKRigDefinition.generated.h"
 
 class UIKRigSolver;
@@ -19,8 +20,6 @@ class IKRIG_API UIKRigDefinition : public UObject
 
 public:
 
-	UIKRigDefinition();
-
 	UPROPERTY()
 	TSoftObjectPtr<UObject> SourceAsset;
 
@@ -30,21 +29,21 @@ public:
 	UPROPERTY(EditAnywhere, Category = BoneSettings)
 	TArray<FName> ExcludedBones;
 
-public:
+private:
+	/** stack of solvers, of varying types, executed in serial fashion where
+	output of prior solve is input to the next */
+	UPROPERTY(EditAnywhere, instanced, Category = "Solvers")
+	TArray<UIKRigSolver*> Solvers;
 
-	void EnsureCreateUniqueGoalName(FName& InOutGoal) const;
+	/** parent/child relationships stored separately from skeleton-specific transforms */
+	UPROPERTY(VisibleAnywhere, Category = "Skeleton")
+	FIKRigHierarchy Hierarchy;
 
-	const TMap<FName, FIKRigGoal>& GetGoals() const
-	{
-		return Goals;
-	}
+	/** default transforms of the reference pose skeleton */
+	UPROPERTY(VisibleAnywhere, Category = "Skeleton")
+	TArray<FTransform> RefPoseTransforms;
 
-	TArray<FName> GetGoalNames() const
-	{
-		TArray<FName> NameArray;
-		Goals.GenerateKeyArray(NameArray);
-		return NameArray;
-	}
+public:
 
 	const TArray<FTransform>& GetReferencePose() const
 	{
@@ -61,84 +60,22 @@ public:
 		return Hierarchy;
 	}
 
-private:
-	
-	/**
-	* Parent/child relationships stored separately from skeleton-specific transforms
-	*/
-	UPROPERTY(VisibleAnywhere, Category = "Skeleton")
-	FIKRigHierarchy Hierarchy;
-
-	/**
-	* These are the default transforms of the reference pose.
-	* 
-	* Transforms are separated from the hierarchical representation
-	* 
-	* These transforms are treated as the common base, from which all
-	* retargeting operations will be compared against.
-	* 
-	* Do no cache these values inside a Solver. Rely instead, only on the 
-	* transforms passed to the Solver::Init()
-	*/
-	UPROPERTY(VisibleAnywhere, Category = "Skeleton")
-	TArray<FTransform> RefPoseTransforms;
-
-	/**
-	* Stack of solvers, of varying types, executed in serial fashion where
-	* output of prior solve is input to the next.
-	*/
-	UPROPERTY(EditAnywhere, instanced, Category = "Solvers")
-	TArray<UIKRigSolver*> Solvers;
+	void GetGoalNamesFromSolvers(TArray<FName>& OutGoalNames) const;
 
-	/** Goals (ie, named transforms) are mapped to solver effectors by name. */
-	UPROPERTY(EditAnywhere, Category = "Goals")
-	TMap<FName, FIKRigGoal> Goals;
+private:
 
 	// BEGIN UObject functions 
 	virtual void PostLoad() override;
-	virtual void BeginDestroy() override;
 	// END UObject functions
 
 #if WITH_EDITOR
-	// https ://drive.google.com/file/d/1EO0Ijojx-0jommdHuZv1JYxfCdoFLd44/view
-
-	/* Goal operators */
-	void UpdateGoal();
 
-	/** hierarchy modifiers */
-	// add a bone. Return false if it already has conflict name or parent is not found
+	//** add a bone. Return false if it already has conflict name or parent is not found */
 	bool AddBone(const FName& InName, const FName& InParent, const FTransform& InGlobalTransform);
-	// remove a bone. Returns false if it is not found
-	// if it has a children, it will remove all the children with it
-	bool RemoveBone(const FName& InName);
-	// rename just change the name
-	// return false if not found
-	bool RenameBone(const FName& InOldName, const FName& InNewName);
-	// reparent to InNewParent
-	// returns false 
-	//	if not found - either InName or InParent
-	//	Or if InNewParent is invalid (i.e. it's a child of InName)
-	bool ReparentBone(const FName& InName, const FName& InNewParent);
-	// reset 
+	//** reset skeleton data */
 	void ResetHierarchy();
-
 	// ensure it's sorted from parent to children
 	void EnsureSortedCorrectly(bool bReSortIfNeeded=false);
-	void Sanitize(); // hopefully we can get rid of this
-
-	// TODO: delegate for hierarchy changes
-	DECLARE_MULTICAST_DELEGATE_OneParam(FBoneAdded, FName /*BoneName*/);
-	DECLARE_MULTICAST_DELEGATE_OneParam(FBoneRemoved, FName /*BoneName*/);
-	DECLARE_MULTICAST_DELEGATE_TwoParams(FBoneRenamed, FName /*OldName*/, FName /*NewName*/);
-	DECLARE_MULTICAST_DELEGATE_TwoParams(FBoneReparented, FName /*Name*/, FName /*ParentName*/);
-
-	DECLARE_MULTICAST_DELEGATE_OneParam(FIKRigDefinitionBeginDestroy, UIKRigDefinition* Definition);
-
-	FBoneAdded BoneAddedDelegate;
-	FBoneRemoved BoneRemovedDelegate;
-	FBoneRenamed BoneRenamedDelegate;
-	FBoneReparented BoneReparentedDelegate;
-	FIKRigDefinitionBeginDestroy IKRigDefinitionBeginDestroy;
 
 #endif // WITH_EDITOR
 
diff --git a/Engine/Plugins/Animation/IKRig/Source/IKRig/Public/IKRigProcessor.h b/Engine/Plugins/Animation/IKRig/Source/IKRig/Public/IKRigProcessor.h
--- a/Engine/Plugins/Animation/IKRig/Source/IKRig/Public/IKRigProcessor.h
+++ b/Engine/Plugins/Animation/IKRig/Source/IKRig/Public/IKRigProcessor.h
@@ -9,44 +9,71 @@
 #include "Drawing/ControlRigDrawInterface.h"
 #include "IKRigProcessor.generated.h"
 
-class UIKRigConstraintSolver;
 class UIKRigDefinition;
 class UIKRigSolver;
 struct FControlRigDrawInterface;
 
-
 UCLASS(BlueprintType)
 class IKRIG_API UIKRigProcessor : public UObject
 {
 	GENERATED_BODY()
 
+public:
+
+	// the runtime for an IKRig, general steps are:
+	// 1. Initialize() with an IKRigDefinition asset
+	// 2. each tick, call SetGoalTransform() and update GlobalBoneTransforms
+	// 3. Call Solve()
+	// 4. Copy back output transforms...
+	// 5. Profit.
+	 
+	// setup
+	void Initialize(UIKRigDefinition* InRigDefinition);
+
+	// run entire stack of solvers
+	void Solve();
+
+	// optionally can be called before Solve() to do deterministic solve
+	void ResetToRefPose();
+
+	// get all global bone transforms
+	FIKRigTransforms& GetCurrentGlobalTransforms();
+
+	// set goal transform by name
+	void SetGoalTransform(const FName& GoalName, const FVector& Position, const FQuat& Rotation);
+
+	// get all goal names in the solver stack
+	void GetGoalNames(TArray<FName>& OutGoalNames) const;
+
+	// total number of goals in solver stack
+	int GetNumGoals() const;
+
+	// get access to the hierarchy the solver stack expects
+	const FIKRigHierarchy* GetHierarchy() const;
+
+	// interface for debug drawing
+	const FControlRigDrawInterface& GetDrawInterface() { return DrawInterface; }
+	
 private:
 
+	bool bInitialized = false;
+
 	UPROPERTY(EditAnywhere, Category = "RigDefinition")
 	UIKRigDefinition* RigDefinition = nullptr;
 
 	UPROPERTY(transient)
-	TArray<FTransform> RefPoseTransforms;
-	FIKRigTransforms Transforms;
+	TArray<FTransform> RefPoseGlobalTransforms;
+
+	UPROPERTY(transient)
+	FIKRigTransforms GlobalBoneTransforms;
 
 	UPROPERTY(transient)
 	TArray<UIKRigSolver*> Solvers;
 
+	UPROPERTY(transient)
 	FControlRigDrawInterface DrawInterface;
-	bool bInitialized = false;
-
-public: 
 
 	UPROPERTY(transient)
 	FIKRigGoalContainer Goals;
-
-	/** setup and run */
-	void Initialize(UIKRigDefinition* InRigDefinition);
-	void Solve();
-
-	FIKRigTransforms& GetTransforms();
-	const FIKRigHierarchy* GetHierarchy() const;
-	void ResetToRefPose();
-	const FControlRigDrawInterface& GetDrawInterface() { return DrawInterface; }
 };
 
diff --git a/Engine/Plugins/Animation/IKRig/Source/IKRig/Public/IKRigSolver.h b/Engine/Plugins/Animation/IKRig/Source/IKRig/Public/IKRigSolver.h
--- a/Engine/Plugins/Animation/IKRig/Source/IKRig/Public/IKRigSolver.h
+++ b/Engine/Plugins/Animation/IKRig/Source/IKRig/Public/IKRigSolver.h
@@ -2,83 +2,66 @@
 
 #pragma once
 
-#include "CoreMinimal.h"
 #include "UObject/ObjectMacros.h"
 #include "UObject/Object.h"
 #include "IKRigDataTypes.h"
 #include "IKRigSolver.generated.h"
 
-struct FIKRigEffector;
 struct FIKRigTarget;
 struct FIKRigTransform;
 struct FIKRigTransforms;
 struct FControlRigDrawInterface;
 
+USTRUCT()
+struct IKRIG_API FIKRigEffectorGoal
+{
+	GENERATED_BODY()
+
+	UPROPERTY(EditAnywhere, Category = FIKRigEffector)
+	FName Goal;
+
+	UPROPERTY(EditAnywhere, Category = FIKRigEffector)
+	FName Bone;
+};
+
 // run time processor 
-UCLASS(Abstract, hidecategories = UObject)
+UCLASS(abstract, hidecategories = UObject)
 class IKRIG_API UIKRigSolver : public UObject
 {
 	GENERATED_BODY()
-
+	
 public:
 
-	/** wraps Solver::Solve() */
-	void SolveInternal(
-		FIKRigTransforms& InOutGlobalTransform, 
-		const FIKRigGoalContainer& Goals,
-		FControlRigDrawInterface* InOutDrawInterface);
-	void AppendGoalNamesToArray(TArray<FName>& OutGoals);
+	UPROPERTY(EditAnywhere, Category = "Definition")
+	bool bEnabled = true;
 
-protected:
+	/** override Init() to setup internal data based on ref pose */
+	virtual void Init(const FIKRigTransforms& InGlobalTransform) PURE_VIRTUAL("Init");
 
-	/** override Init(), Solve() and IsSolverActive() in subclasses */
-	virtual void Init(const FIKRigTransforms& InGlobalTransform) {};
+	/** override Solve() to evaluate new output pose (InOutGlobalTransform) */
 	virtual void Solve(
 		FIKRigTransforms& InOutGlobalTransform,
 		const FIKRigGoalContainer& Goals,
-		FControlRigDrawInterface* InOutDrawInterface) {};
-	virtual bool IsSolverActive() const;
-
-	bool GetGoalForEffector(
-		const FIKRigEffector& InEffector, 
-		const FIKRigGoalContainer &Goals, 
-		FIKRigGoal& OutGoal) const;
-
-	UPROPERTY(EditAnywhere, Category = "Definition")
-	bool bEnabled = true;
-
-	// effector name to goals name map
-	TIKRigEffectorMap<FName> EffectorToGoalName;
-
-private:
-
-	// BEGIN UObject
-	virtual void PostLoad() override;
-	virtual void Serialize(FArchive& Ar) override;
-	// END UObject
-
-	friend class UIKRigController;
-	friend class UIKRigProcessor;
-
-#if WITH_EDITOR
-
-private:
-
-	void RenameGoal(const FName& OldName, const FName& NewName);
-	void EnsureUniqueGoalName(FName& InOutGoalName) const;
-
-	DECLARE_MULTICAST_DELEGATE(FGoalNeedsUpdate);
-	FGoalNeedsUpdate GoalNeedsUpdateDelegate;
-
+		FControlRigDrawInterface* InOutDrawInterface) PURE_VIRTUAL("Solve");
+
+	/** override CollectGoalNames() to tell the processor which goals to collect for this solver
+	 * NOTE: only ADD to OutGoals, do not reset or remove */
+	virtual void CollectGoalNames(TSet<FName>& OutGoals) const PURE_VIRTUAL("CollectGoalNames");
+
+	#if WITH_EDITOR
+	/** override RenameGoal() when UI renames a goal, you can auto-update any effectors
+	 * that were previously mapped to the old goal name. */
+	virtual void RenameGoal(const FName& OldName, const FName& NewName) {};
+	#endif // WITH_EDITOR
+	
 protected:
 
-	FName CreateUniqueGoalName(const TCHAR* Suffix) const;
-	void OnGoalHasBeenUpdated();
-	void EnsureToAddEffector(const FIKRigEffector& InEffector, const FString& InPrefix);
-	void EnsureToRemoveEffector(const FIKRigEffector& InEffector);
-
-	virtual void UpdateEffectors() {};
-
-#endif // WITH_EDITOR
+	static bool GetGoalForEffector(
+		const FIKRigEffectorGoal& InEffector,
+		const FIKRigGoalContainer &Goals,
+		FIKRigGoal& OutGoal)
+	{
+		return Goals.GetGoalByName(InEffector.Goal, OutGoal);
+	}
 };
 
diff --git a/Engine/Plugins/Animation/IKRig/Source/IKRig/Public/Solvers/ConstraintSolver.h b/Engine/Plugins/Animation/IKRig/Source/IKRig/Public/Solvers/ConstraintSolver.h
--- a/Engine/Plugins/Animation/IKRig/Source/IKRig/Public/Solvers/ConstraintSolver.h
+++ b/Engine/Plugins/Animation/IKRig/Source/IKRig/Public/Solvers/ConstraintSolver.h
@@ -5,17 +5,12 @@
 #include "IKRigSolver.h"
 #include "ConstraintSolver.generated.h"
 
-class UIKRigBoneSettings;
 
-
-UCLASS(config = Engine, hidecategories = UObject)
+UCLASS(EditInlineNew, config = Engine, hidecategories = UObject)
 class IKRIG_API UIKRigConstraintSolver : public UIKRigSolver
 {
 	GENERATED_BODY()
 
-	//DECLARE_DELEGATE_TwoParams(FIKRigQueryConstraint, const FName& InConstraintName, UIKRigConstraint& OutConstraint);
-
-
 public: 
 
 	virtual void Init(const FIKRigTransforms& InGlobalTransform) override;
@@ -23,27 +18,6 @@ public:
 		FIKRigTransforms& InOutGlobalTransform,
 		const FIKRigGoalContainer& Goals,
 		FControlRigDrawInterface* InOutDrawInterface) override;
-
-	// this can be utilized for just subset of constraints
-	//void SolverConstraints(const TArray<FName>& ConstraintsList, FIKRigTransformModifier& InOutGlobalTransform, FControlRigDrawInterface* InOutDrawInterface);
-
-	// register/unregister query function
-	//void RegisterQueryConstraintHandler(const FIKRigQueryConstraint& InQueryConstraintHandler);
-	//void UnregisterQueryConstraintHandler();
-
-	//void SetConstraintDefinition(UIKRigConstraintDefinition* InConstraintDefinition);
-
-private:
-	// delegate
-	//FIKRigQueryConstraint QueryConstraintHandler;
-
-	// current active profile
-	//UPROPERTY(transient)
-	//FName ActiveProfile;
-
-	// instanced information you can mutate from ConstraintDefinition
-	// we want uproperty, so that it doesn't GC-ed
-	//UPROPERTY(transient)
-	//TMap<FName, FIKRigConstraintProfile> ConstraintProfiles;
+	virtual void CollectGoalNames(TSet<FName>& OutGoals) const override;
 };
 
diff --git a/Engine/Plugins/Animation/IKRig/Source/IKRig/Public/Solvers/TransformSolver.h b/Engine/Plugins/Animation/IKRig/Source/IKRig/Public/Solvers/TransformSolver.h
--- a/Engine/Plugins/Animation/IKRig/Source/IKRig/Public/Solvers/TransformSolver.h
+++ b/Engine/Plugins/Animation/IKRig/Source/IKRig/Public/Solvers/TransformSolver.h
@@ -6,15 +6,15 @@
 #include "TransformSolver.generated.h"
 
 
-UCLASS()
+UCLASS(EditInlineNew)
 class IKRIG_API UTransformSolver : public UIKRigSolver
 {
 	GENERATED_BODY()
 
-public:
-
 	UTransformSolver();
 
+public:
+
 	UPROPERTY(EditAnywhere, Category = "Solver")
 	bool bEnablePosition = true;
 
@@ -22,7 +22,7 @@ public:
 	bool bEnableRotation = true;
 
 	UPROPERTY(EditAnywhere, Category = "Solver")
-	FIKRigEffector Effector;
+	FIKRigEffectorGoal Effector;
 
 protected:
 	
@@ -31,17 +31,6 @@ protected:
 		FIKRigTransforms& InOutGlobalTransform, 
 		const FIKRigGoalContainer& Goals,
 		FControlRigDrawInterface* InOutDrawInterface) override;
-	virtual bool IsSolverActive() const override;
-
-private:
-#if WITH_EDITOR
-	virtual void UpdateEffectors() override;
-
-	// UObject interface
-	virtual void PostEditChangeChainProperty(struct FPropertyChangedChainEvent& PropertyChangedEvent) override;
-	// UObject interface
-#endif // WITH_EDITOR
-
-	const FString TransformTargetName;
+	virtual void CollectGoalNames(TSet<FName>& OutGoals) const override;
 };
 
diff --git a/Engine/Plugins/Animation/IKRig/Source/IKRigEditor/Private/AssetTypeActions_IKRigDefinition.h b/Engine/Plugins/Animation/IKRig/Source/IKRigEditor/Private/AssetTypeActions_IKRigDefinition.h
--- a/Engine/Plugins/Animation/IKRig/Source/IKRigEditor/Private/AssetTypeActions_IKRigDefinition.h
+++ b/Engine/Plugins/Animation/IKRig/Source/IKRigEditor/Private/AssetTypeActions_IKRigDefinition.h
@@ -10,7 +10,7 @@ class FAssetTypeActions_IKRigDefinition : public FAssetTypeActions_Base
 {
 public:
 	// IAssetTypeActions Implementation
-	virtual FText GetName() const override { return NSLOCTEXT("AssetTypeActions", "AssetTypeActions_IKRigDefinition", "IK Rig Definition"); }
+	virtual FText GetName() const override { return NSLOCTEXT("AssetTypeActions", "AssetTypeActions_IKRigDefinition", "IK Rig"); }
 	virtual void GetActions(const TArray<UObject*>& InObjects, struct FToolMenuSection& Section) override;
 	virtual FColor GetTypeColor() const override { return FColor(255, 192, 128); }
 	virtual UClass* GetSupportedClass() const override;
diff --git a/Engine/Plugins/Animation/IKRig/Source/IKRigEditor/Private/IKRigController.cpp b/Engine/Plugins/Animation/IKRig/Source/IKRigEditor/Private/IKRigController.cpp
--- a/Engine/Plugins/Animation/IKRig/Source/IKRigEditor/Private/IKRigController.cpp
+++ b/Engine/Plugins/Animation/IKRig/Source/IKRigEditor/Private/IKRigController.cpp
@@ -9,33 +9,33 @@
 
 #define LOCTEXT_NAMESPACE	"IKRigController"
 
+
+// -------------------------------------------------------
+// CONTROLLER <-> IKRIGDEFINITION CONNECTION
+//
+
 TMap<UIKRigDefinition*, UIKRigController*> UIKRigController::DefinitionToControllerMap;
 
-// currently it's not clear the lifecycle of this map
-// usually during editor, they will be present, 
 UIKRigController* UIKRigController::GetControllerByRigDefinition(UIKRigDefinition* InIKRigDefinition)
 {
-	if (InIKRigDefinition)
+	if (!InIKRigDefinition)
 	{
-		UIKRigController** Controller = DefinitionToControllerMap.Find(InIKRigDefinition);
-		if (Controller)
-		{
-			return *Controller;
-		}
-		else
-		{
-			UIKRigController* NewController = NewObject<UIKRigController>();
-			DefinitionToControllerMap.Add(InIKRigDefinition) = NewController;
-			NewController->SetIKRigDefinition(InIKRigDefinition);
-			InIKRigDefinition->IKRigDefinitionBeginDestroy.AddStatic(&UIKRigController::RemoveControllerByRigDefinition);
-			return NewController;
-		}
+		return nullptr;
+	}
+	
+	UIKRigController** Controller = DefinitionToControllerMap.Find(InIKRigDefinition);
+	if (Controller)
+	{
+		return *Controller;
 	}
 
-	return nullptr;
+	UIKRigController* NewController = NewObject<UIKRigController>();
+	DefinitionToControllerMap.Add(InIKRigDefinition) = NewController;
+	NewController->SetIKRigDefinition(InIKRigDefinition);
+	return NewController;
 }
 
-// this shoudl be called by IKRigDefinition::BeginDestroy;
+// this should be called by IKRigDefinition::BeginDestroy;
 void UIKRigController::RemoveControllerByRigDefinition(UIKRigDefinition* InIKRigDefinition)
 {
 	DefinitionToControllerMap.Remove(InIKRigDefinition);
@@ -56,254 +56,129 @@ void UIKRigController::AddReferencedObjects(UObject* InThis, FReferenceCollector
 	}
 }
 
-void UIKRigController::BeginDestroy() 
-{
-	if (IKRigDefinition)
-	{
-		for (int32 Index = 0; Index < IKRigDefinition->Solvers.Num(); ++Index)
-		{
-			UninitializeSolver(IKRigDefinition->Solvers[Index]);
-		}
-	}
-
-	Super::BeginDestroy();
-}
-
- // IKRigDefinition set up
 void UIKRigController::SetIKRigDefinition(UIKRigDefinition* InIKRigDefinition)
 {
-	if (IKRigDefinition)
-	{
-		for (int32 Index = 0; Index < IKRigDefinition->Solvers.Num(); ++Index)
-		{
-			UninitializeSolver(IKRigDefinition->Solvers[Index]);
-		}
-	}
-
 	IKRigDefinition = InIKRigDefinition;
-
-	if (IKRigDefinition)
-	{
-		for (int32 Index = 0; Index < IKRigDefinition->Solvers.Num(); ++Index)
-		{
-			InitializeSolver(IKRigDefinition->Solvers[Index]);
-		}
-	}
 }
 
-void UIKRigController::InitializeSolver(UIKRigSolver* Solver)
-{
-	if (Solver)
-	{
-		// we want to register this controller delegate
-		// if there are multiple controllers managing one IKRigDefinition
-		// we want to ensure that works 
-		FDelegateHandle DelegateHandle = Solver->GoalNeedsUpdateDelegate.AddUObject(this, &UIKRigController::UpdateGoal);
-		SolverDelegateHandles.Add(Solver) = DelegateHandle;
+// -------------------------------------------------------
+// SKELETON
+//
 
-		Solver->UpdateEffectors();
-	}
+const FIKRigHierarchy* UIKRigController::GetHierarchy() const
+{
+	return (IKRigDefinition) ? &IKRigDefinition->GetHierarchy() : nullptr;
 }
 
-void UIKRigController::UninitializeSolver(UIKRigSolver* Solver)
+const TArray<FTransform>& UIKRigController::GetRefPoseTransforms() const
 {
-	if (Solver)
+	if (IKRigDefinition)
 	{
-		FDelegateHandle* DelegateHandle = SolverDelegateHandles.Find(Solver);
-		if (ensure(DelegateHandle))
-		{
-			Solver->GoalNeedsUpdateDelegate.Remove(*DelegateHandle);
-		}
+		return IKRigDefinition->GetReferencePose();
 	}
+
+	static TArray<FTransform> Dummy;
+	return Dummy;
 }
 
-// hierarchy operators
 void UIKRigController::SetSkeleton(const FReferenceSkeleton& InSkeleton)
 {
-	if (IKRigDefinition)
+	if (!IKRigDefinition)
 	{
-		FScopedTransaction Transaction(LOCTEXT("SetSkeleton_Label", "Set Skeleton"));
-		IKRigDefinition->Modify();
-
-		IKRigDefinition->ResetHierarchy();
+		return;
+	}
 
-		const TArray<FMeshBoneInfo>& RefBoneInfo = InSkeleton.GetRefBoneInfo();
-		TArray<FTransform> RefPoseInCS;
-		FAnimationRuntime::FillUpComponentSpaceTransforms(InSkeleton, InSkeleton.GetRefBonePose(), RefPoseInCS);
-		ensure(RefPoseInCS.Num() == RefBoneInfo.Num());
-		for (int32 Index=0; Index<RefBoneInfo.Num(); ++Index)
-		{
-			int32 ParentIndex = RefBoneInfo[Index].ParentIndex;
-			ensure (IKRigDefinition->AddBone(RefBoneInfo[Index].Name, (ParentIndex != INDEX_NONE)? RefBoneInfo[ParentIndex].Name : NAME_None, RefPoseInCS[Index]));
-		}
+	FScopedTransaction Transaction(LOCTEXT("SetSkeleton_Label", "Set Skeleton"));
+	IKRigDefinition->Modify();
 
-		ensure(IKRigDefinition->Hierarchy.GetNum() == IKRigDefinition->RefPoseTransforms.Num());
-	}
-}
+	IKRigDefinition->ResetHierarchy();
 
-bool UIKRigController::AddBone(const FName& InName, const FName& InParent, const FTransform& InGlobalTransform)
-{
-	if (IKRigDefinition)
+	const TArray<FMeshBoneInfo>& RefBoneInfo = InSkeleton.GetRefBoneInfo();
+	TArray<FTransform> RefPoseInCS;
+	FAnimationRuntime::FillUpComponentSpaceTransforms(InSkeleton, InSkeleton.GetRefBonePose(), RefPoseInCS);
+	ensure(RefPoseInCS.Num() == RefBoneInfo.Num());
+	for (int32 Index = 0; Index < RefBoneInfo.Num(); ++Index)
 	{
-		FScopedTransaction Transaction(LOCTEXT("AddBone_Label", "Add Bone"));
-		IKRigDefinition->Modify();
-
-		return IKRigDefinition->AddBone(InName, InParent, InGlobalTransform);
+		int32 ParentIndex = RefBoneInfo[Index].ParentIndex;
+		ensure(IKRigDefinition->AddBone(RefBoneInfo[Index].Name, (ParentIndex != INDEX_NONE) ? RefBoneInfo[ParentIndex].Name : NAME_None, RefPoseInCS[Index]));
 	}
 
-	return false;
+	ensure(IKRigDefinition->Hierarchy.GetNum() == IKRigDefinition->RefPoseTransforms.Num());
 }
 
-bool UIKRigController::RemoveBone(const FName& InName)
+bool UIKRigController::AddBone(const FName& InName, const FName& InParent, const FTransform& InGlobalTransform)
 {
-	if (IKRigDefinition)
+	if (!IKRigDefinition)
 	{
-		FScopedTransaction Transaction(LOCTEXT("RemoveBone_Label", "Remove Bone"));
-		IKRigDefinition->Modify();
-
-		return IKRigDefinition->RemoveBone(InName);
+		return false;
 	}
 
-	return false;
-}
-bool UIKRigController::RenameBone(const FName& InOldName, const FName& InNewName)
-{
-	if (IKRigDefinition)
-	{
-		FScopedTransaction Transaction(LOCTEXT("RenameBone_Label", "Rename Bone"));
-		IKRigDefinition->Modify();
-
-		return IKRigDefinition->RenameBone(InOldName, InNewName);
-	}
+	FScopedTransaction Transaction(LOCTEXT("AddBone_Label", "Add Bone"));
+	IKRigDefinition->Modify();
 
-	return false;
+	return IKRigDefinition->AddBone(InName, InParent, InGlobalTransform);
 }
-bool UIKRigController::ReparentBone(const FName& InName, const FName& InNewParent)
-{
-	if (IKRigDefinition)
-	{
-		FScopedTransaction Transaction(LOCTEXT("ReparentBone_Label", "Reparent Bone"));
-		IKRigDefinition->Modify();
 
-		return IKRigDefinition->ReparentBone(InName, InNewParent);
-	}
 
-	return false;
-}
 void UIKRigController::ResetHierarchy()
 {
-	if (IKRigDefinition)
+	if (!IKRigDefinition)
 	{
-		FScopedTransaction Transaction(LOCTEXT("ResetHierarchy_Label", "Reset Hierarchy"));
-		IKRigDefinition->Modify();
-
-		IKRigDefinition->ResetHierarchy();
+		return;
 	}
-}
-
-// solver operators
-UIKRigSolver* UIKRigController::AddSolver(TSubclassOf<UIKRigSolver> InIKRigSolverClass)
-{
-	if (IKRigDefinition)
-	{
-		FScopedTransaction Transaction(LOCTEXT("AddSolver_Label", "Add Solver"));
-		IKRigDefinition->Modify();
 
-		UIKRigSolver* NewSolver = NewObject<UIKRigSolver>(IKRigDefinition, InIKRigSolverClass);
-		check(NewSolver);
-
-		// todo: set delegate for the goal update
-		IKRigDefinition->Solvers.Add(NewSolver);
-		InitializeSolver(NewSolver);
-		IKRigDefinition->UpdateGoal();
-		return NewSolver;
-	}
+	FScopedTransaction Transaction(LOCTEXT("ResetHierarchy_Label", "Reset Hierarchy"));
+	IKRigDefinition->Modify();
 
-	return nullptr;
+	IKRigDefinition->ResetHierarchy();
 }
 
-int32 UIKRigController::GetNumSolvers() const
-{
-	if (IKRigDefinition)
-	{
-		return IKRigDefinition->Solvers.Num();
-	}
-
-	return 0;
-}
+// -------------------------------------------------------
+// SOLVERS
+//
 
-UIKRigSolver* UIKRigController::GetSolver(int32 Index) const
+UIKRigSolver* UIKRigController::AddSolver(TSubclassOf<UIKRigSolver> InIKRigSolverClass)
 {
-	if (IKRigDefinition && IKRigDefinition->Solvers.IsValidIndex(Index))
+	if (!IKRigDefinition)
 	{
-		return IKRigDefinition->Solvers[Index];
+		return nullptr;
 	}
 
-	return nullptr;
-}
+	FScopedTransaction Transaction(LOCTEXT("AddSolver_Label", "Add Solver"));
+	IKRigDefinition->Modify();
 
-void UIKRigController::RemoveSolver(UIKRigSolver* SolverToDelete)
-{
-	if (IKRigDefinition && SolverToDelete)
-	{
-		FScopedTransaction Transaction(LOCTEXT("RemoveSolver_Label", "Remove Solver"));
-		IKRigDefinition->Modify();
+	UIKRigSolver* NewSolver = NewObject<UIKRigSolver>(IKRigDefinition, InIKRigSolverClass);
+	check(NewSolver);
 
-		UninitializeSolver(SolverToDelete);
-		IKRigDefinition->Solvers.Remove(SolverToDelete);
-		IKRigDefinition->UpdateGoal();
-	}
+	IKRigDefinition->Solvers.Add(NewSolver);
+	return NewSolver;
 }
 
-bool UIKRigController::ValidateSolver(UIKRigSolver* const Solver) const
+int32 UIKRigController::GetNumSolvers() const
 {
-	return (IKRigDefinition && Solver && IKRigDefinition->Solvers.Find(Solver) != INDEX_NONE);
+	return IKRigDefinition ? IKRigDefinition->Solvers.Num() : 0;
 }
 
-void UIKRigController::UpdateGoal()
+UIKRigSolver* UIKRigController::GetSolver(int32 Index) const
 {
-	if (IKRigDefinition)
-	{
-		FScopedTransaction Transaction(LOCTEXT("UpdateGoal_Label", "Update Goal"));
-		IKRigDefinition->Modify();
-
-		IKRigDefinition->UpdateGoal();
-		OnGoalModified.Broadcast();
-	}
+	bool ValidSolver = IKRigDefinition && IKRigDefinition->Solvers.IsValidIndex(Index);
+	return ValidSolver ? IKRigDefinition->Solvers[Index] : nullptr;
 }
 
-FName UIKRigController::GetGoalName(UIKRigSolver* InSolver, const FIKRigEffector& InEffector)
+void UIKRigController::RemoveSolver(UIKRigSolver* SolverToDelete)
 {
-	if (InSolver)
+	if (!(IKRigDefinition && SolverToDelete))
 	{
-		FName* GoalName = InSolver->EffectorToGoalName.Find(InEffector);
-		if (GoalName)
-		{
-			return *GoalName;
-		}
+		return;
 	}
 
-	return NAME_None;
-}
-
-void UIKRigController::SetGoalName(UIKRigSolver* InSolver, const FIKRigEffector& InEffector, const FName& NewGoalName)
-{
-	if (InSolver)
-	{
-		FScopedTransaction Transaction(LOCTEXT("SetGoalName_Label", "Set Goal Name"));
-		InSolver->Modify();
-
-		FName* GoalName = InSolver->EffectorToGoalName.Find(InEffector);
-		if (GoalName)
-		{
-			*GoalName = NewGoalName;
+	FScopedTransaction Transaction(LOCTEXT("RemoveSolver_Label", "Remove Solver"));
+	IKRigDefinition->Modify();
 
-			UpdateGoal();
-		}
-	}
+	IKRigDefinition->Solvers.Remove(SolverToDelete);
 }
 
-UIKRigBoneSetting* UIKRigController::AddBoneSetting(TSubclassOf<UIKRigBoneSetting> NewSettingType)
+UIKRigBoneSetting* UIKRigController::AddBoneSetting(TSubclassOf<UIKRigBoneSetting> NewSettingType) const
 {
 	
 	FScopedTransaction Transaction(LOCTEXT("AddSetting_Label", "Add Bone Setting"));
@@ -319,54 +194,34 @@ UIKRigBoneSetting* UIKRigController::AddBoneSetting(TSubclassOf<UIKRigBoneSettin
 	return nullptr;	
 }
 
-// goal operators
-void UIKRigController::QueryGoals(TArray<FName>& OutGoals) const
-{
-	if (IKRigDefinition)
-	{
-		IKRigDefinition->GetGoals().GenerateKeyArray(OutGoals);
-	}
-}
+// -----------------------------
+// GOALS
+//
 
-void UIKRigController::RenameGoal(const FName& OldName, const FName& NewName)
+void UIKRigController::GetGoalNames(TArray<FName>& OutGoals) const
 {
-	// ensure we don't have used NewDisplayName 
-	if (IKRigDefinition)
+	if (!IKRigDefinition)
 	{
-		FScopedTransaction Transaction(LOCTEXT("RenameGoal_Label", "Rename Goal"));
-		IKRigDefinition->Modify();
-
-		for (UIKRigSolver* Solver : IKRigDefinition->Solvers)
-		{
-			if (Solver)
-			{
-				Solver->RenameGoal(OldName, NewName);
-			}
-		}
-		
-		// update goal list
-		IKRigDefinition->UpdateGoal();
+		return;
 	}
+
+	IKRigDefinition->GetGoalNamesFromSolvers(OutGoals);
 }
 
-FIKRigGoal* UIKRigController::GetGoal(const FName& InGoalName) 
+void UIKRigController::RenameGoal(const FName& OldName, const FName& NewName) const
 {
-	if (IKRigDefinition)
+	if (!IKRigDefinition)
 	{
-		return IKRigDefinition->Goals.Find(InGoalName);
+		return;
 	}
 
-	return nullptr;
-}
+	FScopedTransaction Transaction(LOCTEXT("RenameGoal_Label", "Rename Goal"));
+	IKRigDefinition->Modify();
 
-const FIKRigGoal* UIKRigController::GetGoal(const FName& InGoalName) const
-{
-	if (IKRigDefinition)
+	for (UIKRigSolver* Solver : IKRigDefinition->Solvers)
 	{
-		return IKRigDefinition->Goals.Find(InGoalName);
+		Solver->RenameGoal(OldName, NewName);
 	}
-
-	return nullptr;
 }
 
 #undef LOCTEXT_NAMESPACE
\ No newline at end of file
diff --git a/Engine/Plugins/Animation/IKRig/Source/IKRigEditor/Private/IKRigDefinitionDetails.cpp b/Engine/Plugins/Animation/IKRig/Source/IKRigEditor/Private/IKRigDefinitionDetails.cpp
--- a/Engine/Plugins/Animation/IKRig/Source/IKRigEditor/Private/IKRigDefinitionDetails.cpp
+++ b/Engine/Plugins/Animation/IKRig/Source/IKRigEditor/Private/IKRigDefinitionDetails.cpp
@@ -5,21 +5,14 @@
 #include "AssetData.h"
 #include "DetailLayoutBuilder.h"
 #include "DetailCategoryBuilder.h"
-#include "IDetailsView.h"
 
 #include "IKRigDefinition.h"
 #include "IKRigController.h"
-#include "IKRigSolver.h"
-#include "IKRigBoneSetting.h"
 
 #include "ScopedTransaction.h"
 #include "PropertyCustomizationHelpers.h"
-#include "Kismet2/SClassPickerDialog.h"
 #include "Animation/Skeleton.h"
-#include "ClassViewerFilter.h"
 #include "Engine/SkeletalMesh.h"
-#include "Modules/ModuleManager.h"
-#include "Widgets/Input/SEditableTextBox.h"
 
 #define LOCTEXT_NAMESPACE	"IKRigDefinitionDetails"
 
@@ -30,10 +23,6 @@ TSharedRef<IDetailCustomization> FIKRigDefinitionDetails::MakeInstance()
 
 FIKRigDefinitionDetails::~FIKRigDefinitionDetails()
 {
-	if (ObjectChangedDelegate.IsValid())
-	{
-		FCoreUObjectDelegates::OnObjectPropertyChanged.Remove(ObjectChangedDelegate);
-	}
 }
 void FIKRigDefinitionDetails::CustomizeDetails(const TSharedPtr<IDetailLayoutBuilder>& DetailBuilder)
 {
@@ -43,9 +32,9 @@ void FIKRigDefinitionDetails::CustomizeDetails(const TSharedPtr<IDetailLayoutBui
 
 void FIKRigDefinitionDetails::CustomizeDetails(IDetailLayoutBuilder& DetailBuilder)
 {
+	// get selected IKRigDefinition asset
 	const TArray< TWeakObjectPtr<UObject> >& SelectedObjectsList = DetailBuilder.GetSelectedObjects();
 	TArray< TWeakObjectPtr<UIKRigDefinition> > SelectedIKRigDefinitions;
-
 	for (auto SelectionIt = SelectedObjectsList.CreateConstIterator(); SelectionIt; ++SelectionIt)
 	{
 		if (UIKRigDefinition* TestIKRigDefinition = Cast<UIKRigDefinition>(SelectionIt->Get()))
@@ -53,225 +42,78 @@ void FIKRigDefinitionDetails::CustomizeDetails(IDetailLayoutBuilder& DetailBuild
 			SelectedIKRigDefinitions.Add(TestIKRigDefinition);
 		}
 	}
-
-	// we only support 1 asset for now
 	if (SelectedIKRigDefinitions.Num() > 1)
 	{
-		return;
+		return; // we only support 1 asset for now
 	}
 
+	// resolve pointer to selected IKRigDefinition
 	IKRigDefinition = SelectedIKRigDefinitions[0];
-
 	if (!IKRigDefinition.IsValid())
 	{
 		return;
 	}
 
-	// create controller
+	// store controller class (MVC pattern) for UI callbacks to interact with
 	IKRigController = UIKRigController::GetControllerByRigDefinition(IKRigDefinition.Get());
 	
-	ObjectChangedDelegate = FCoreUObjectDelegates::OnObjectPropertyChanged.AddRaw(this, &FIKRigDefinitionDetails::OnObjectPostEditChange);
-	/////////////////////////////////////////////////////////////////////////////////
-	// skeleton set up
-	/////////////////////////////////////////////////////////////////////////////////
-	IDetailCategoryBuilder& HierarchyCategory = DetailBuilder.EditCategory("Skeleton");
- 
-	SelectedAsset = IKRigDefinition->SourceAsset.Get();
-
-	HierarchyCategory.AddCustomRow(FText::FromString("ChangeSkeleton"))
-	.NameContent()
-	[
-		SNew(STextBlock)
-		.Font(IDetailLayoutBuilder::GetDetailFont())
-		.Text(LOCTEXT("SelectSourceSkeleton", "Source Skeleton"))
-	]
-	.ValueContent()
-	[
-		SNew(SBorder)
-		.BorderImage(FCoreStyle::Get().GetBrush("ToolPanel.GroupBorder"))
-		.BorderBackgroundColor(FLinearColor::Gray) // Darken the outer border
-		[
-			SNew(SVerticalBox)
-			+ SVerticalBox::Slot()
-			.AutoHeight()
-			.HAlign(HAlign_Center)
-			.Padding(2, 2)
-			[
-				SNew(SBox)
-				.WidthOverride(300)
-				[
-					SNew(SObjectPropertyEntryBox)
-					.ObjectPath(this, &FIKRigDefinitionDetails::GetCurrentSourceAsset)
-					.OnShouldFilterAsset(this, &FIKRigDefinitionDetails::ShouldFilterAsset)
-					.OnObjectChanged(this, &FIKRigDefinitionDetails::OnAssetSelected)
-					.AllowClear(false)
-					.DisplayUseSelected(true)
-					.DisplayBrowse(true)
-				]
-			]
-			+ SVerticalBox::Slot()
-			.AutoHeight()
-			.HAlign(HAlign_Center)
-			.Padding(2, 2)
-			[
-				SNew(SButton)
-				.ContentPadding(3)
-				.IsEnabled(this, &FIKRigDefinitionDetails::CanImport)
-				.OnClicked(this, &FIKRigDefinitionDetails::OnImportHierarchy)
-				.ToolTipText(LOCTEXT("OnImportHierarchyTooltip", "Set skeleton to selected asset. This replaces existing skeleton."))
-				[
-					SNew(STextBlock)
-					.Font(IDetailLayoutBuilder::GetDetailFont())
-					.Text(LOCTEXT("UpdateHierarchyTitle", "Update"))
-				]
-			]
-		]
-	];
-
- 	IDetailCategoryBuilder& SolverCategory = DetailBuilder.EditCategory("Solvers");
+	//
+	// EDIT SKELETON
+	//
 
-	SolverCategory.AddCustomRow(FText::FromString("AddSolver"))
+	IDetailCategoryBuilder& SkeletonCategory = DetailBuilder.EditCategory("Skeleton");
+	SkeletonCategory.AddCustomRow(FText::FromString("UpdateSkeleton"))
 	.NameContent()
 	[
 		SNullWidget::NullWidget
 	]
 	.ValueContent()
+		.MaxDesiredWidth(0.0f)
+		.MinDesiredWidth(150.0f)
 	[
 		SNew(SButton)
 		.ContentPadding(3)
-		.OnClicked(this, &FIKRigDefinitionDetails::OnShowSolverClassPicker)
-		.ToolTipText(LOCTEXT("OnShowSolverListTooltip", "Select Solver to Add"))
-		[
-			SNew(STextBlock)
-			.Font(IDetailLayoutBuilder::GetDetailFont())
-			.Text(LOCTEXT("ShowSolverList", "Add Solver"))
-		]
+		.IsEnabled(this, &FIKRigDefinitionDetails::CanImport)
+		.OnClicked(this, &FIKRigDefinitionDetails::OnImportHierarchy)
+		.ToolTipText(LOCTEXT("OnImportHierarchyTooltip", "Set skeleton to selected asset. This replaces existing skeleton."))
+		.Text(LOCTEXT("UpdateHierarchyTitle", "Update Skeleton"))
+		.HAlign(HAlign_Center)
 	];
 
-	IDetailCategoryBuilder& BoneSettingCategory = DetailBuilder.EditCategory("BoneSettings");
-	BoneSettingCategory.AddCustomRow(FText::FromString("AddBoneSetting"))
+	
+	SkeletonCategory.AddCustomRow(FText::FromString("Hierarchy"))
 	.NameContent()
 	[
-		SNullWidget::NullWidget
+		SNew(STextBlock)
+		.Font(IDetailLayoutBuilder::GetDetailFont())
+		.Text(LOCTEXT("SelectSourceSkeleton", "Source Skeleton"))
 	]
 	.ValueContent()
+		.MaxDesiredWidth(0.0f)
+		.MinDesiredWidth(200.0f)
 	[
-		SNew(SButton)
-		.ContentPadding(3)
-		.OnClicked(this, &FIKRigDefinitionDetails::OnShowBoneSettingsClassPicker)
-		.ToolTipText(LOCTEXT("OnShowBoneSettingListTooltip", "Select Bone Setting to Add"))
+		SNew(SBorder)
+		.BorderImage(FCoreStyle::Get().GetBrush("ToolPanel.GroupBorder"))
+		.BorderBackgroundColor(FLinearColor::Gray) // Darken the outer border
 		[
-			SNew(STextBlock)
-			.Font(IDetailLayoutBuilder::GetDetailFont())
-			.Text(LOCTEXT("ShowBoneSettingList", "Add Bone Setting"))
+			SNew(SBox)
+			[
+				SNew(SObjectPropertyEntryBox)
+				.ObjectPath(this, &FIKRigDefinitionDetails::GetCurrentSourceAsset)
+				.OnShouldFilterAsset(this, &FIKRigDefinitionDetails::ShouldFilterAsset)
+				.OnObjectChanged(this, &FIKRigDefinitionDetails::OnAssetSelected)
+				.AllowClear(false)
+				.DisplayUseSelected(true)
+				.DisplayBrowse(true)
+			]
 		]
 	];
-
-	GoalPropertyHandle = DetailBuilder.GetProperty(TEXT("Goals"));
-
-	TArray<FName> GoalNames;
-	IKRigController->QueryGoals(GoalNames);
-
-	GoalListNames.SetNumZeroed(GoalNames.Num());
-
-	for (int32 Index = 0; Index < GoalNames.Num(); ++Index)
-	{
-		GoalListNames[Index] = MakeShareable(new FGoalNameListItem(GoalNames[Index]));
-	}
-
-	// I need to think about goal modified event OR just IKRigAssetMOdified event to update this
-	// for now i'm commenting it out
-// 	IDetailCategoryBuilder& GoalPropertyGroup = DetailBuilder.EditCategory("Goals");
-// 	GoalPropertyGroup.AddCustomRow(LOCTEXT("GoalsTitleLabel", "Goals"))
-// 	.NameContent()
-// 	[
-// 		GoalPropertyHandle->CreatePropertyNameWidget()
-// 	]
-// 	.ValueContent()
-// 	[
-// 		SAssignNew(GoalListView, SListView<FGoalNameListItemPtr>)
-// 		.ListItemsSource(&GoalListNames)
-// 		.OnGenerateRow(this, &FIKRigDefinitionDetails::OnGenerateWidgetForGoals)
-// 	];
-
-//	GoalPropertyHandle->MarkHiddenByCustomization();
-}
-
-class FIKRigClassFilter : public IClassViewerFilter
-{
-public:
-	/** All children of these classes will be included unless filtered out by another setting. */
-	TSet< const UClass* > AllowedChildrenOfClasses;
-
-	TSet< const UClass* > DisallowedClasses;
-
-	/** Disallowed class flags. */
-	EClassFlags DisallowedClassFlags;
-
-	virtual bool IsClassAllowed(const FClassViewerInitializationOptions& InInitOptions, const UClass* InClass, TSharedRef< FClassViewerFilterFuncs > InFilterFuncs) override
-	{
-		return !InClass->HasAnyClassFlags(DisallowedClassFlags) && InFilterFuncs->IfInClassesSet(DisallowedClasses, InClass) == EFilterReturn::Failed
-			&& InFilterFuncs->IfInChildOfClassesSet(AllowedChildrenOfClasses, InClass) != EFilterReturn::Failed;
-	}
-
-	virtual bool IsUnloadedClassAllowed(const FClassViewerInitializationOptions& InInitOptions, const TSharedRef< const IUnloadedBlueprintData > InUnloadedClassData, TSharedRef< FClassViewerFilterFuncs > InFilterFuncs) override
-	{
-		return !InUnloadedClassData->HasAnyClassFlags(DisallowedClassFlags) && InFilterFuncs->IfInClassesSet(DisallowedClasses, InUnloadedClassData) == EFilterReturn::Failed
-			&& InFilterFuncs->IfInChildOfClassesSet(AllowedChildrenOfClasses, InUnloadedClassData) != EFilterReturn::Failed;
-	}
-};
-
-// choose class for them
-UClass* SelectClass(UClass* ClassType, const FText& TitleText)
-{
-	FClassViewerModule& ClassViewerModule = FModuleManager::LoadModuleChecked<FClassViewerModule>("ClassViewer");
-	FClassViewerInitializationOptions Options;
-	Options.Mode = EClassViewerMode::ClassPicker;
-	Options.DisplayMode = EClassViewerDisplayMode::TreeView;
-	Options.bShowObjectRootClass = false;
-	Options.bExpandRootNodes = true;
-	Options.bShowUnloadedBlueprints = true;
-	TSharedPtr<FIKRigClassFilter> Filter = MakeShareable(new FIKRigClassFilter);
-	Options.ClassFilter = Filter;
-
-	Filter->DisallowedClassFlags = CLASS_Abstract | CLASS_Deprecated | CLASS_NewerVersionExists | CLASS_Transient;
-	Filter->AllowedChildrenOfClasses.Add(ClassType);
-	Filter->DisallowedClasses.Add(ClassType);
-
-	UClass* ChosenClass = nullptr;
-	const bool bPressedOk = SClassPickerDialog::PickClass(TitleText, Options, ChosenClass, ClassType);
-
-	if (bPressedOk)
-	{
-		return ChosenClass;
-	}
-
-	return nullptr;
 }
 
-
-
-FReply FIKRigDefinitionDetails::OnShowSolverClassPicker()
-{
-	UClass* ChosenClass = SelectClass(UIKRigSolver::StaticClass(), LOCTEXT("SelectSolverClass", "Select Solver Class"));
-	if (ChosenClass)
-	{
-		IKRigController->AddSolver(ChosenClass);
-	}
-
-	return FReply::Handled();
-}
-
-FReply FIKRigDefinitionDetails::OnShowBoneSettingsClassPicker()
-{
-	UClass* ChosenClass = SelectClass(UIKRigBoneSetting::StaticClass(), LOCTEXT("SelectBoneSettingClass", "Select Bone Setting Class"));
-	if (ChosenClass)
-	{
-		IKRigController->AddBoneSetting(ChosenClass);
-	}
-
-	return FReply::Handled();
-}
+//------------------------------------------------------------------
+//
+// HIERARCHY
+//
 
 bool FIKRigDefinitionDetails::CanImport() const
 {
@@ -328,59 +170,4 @@ FReply FIKRigDefinitionDetails::OnImportHierarchy()
 	return FReply::Handled();
 }
 
-void FIKRigDefinitionDetails::OnObjectPostEditChange(UObject* Object, FPropertyChangedEvent& InPropertyChangedEvent)
-{
-// 	if (Object == IKRigDefinition || Object->GetOuter())
-// 	{
-// 		IDetailLayoutBuilder* DetailLayoutBuilder = DetailBuilderWeakPtr.Pin().Get();
-// 		if (DetailLayoutBuilder)
-// 		{
-// 			DetailLayoutBuilder->ForceRefreshDetails();
-// 		}
-// 	}
-}
-
-TSharedRef<ITableRow> FIKRigDefinitionDetails::OnGenerateWidgetForGoals(FGoalNameListItemPtr InItem, const TSharedRef<STableViewBase>& OwnerTable)
-{
-	return SNew(STableRow<FGoalNameListItemPtr>, OwnerTable)
-		[
-			SNew(SHorizontalBox)
-			+SHorizontalBox::Slot()
-			.AutoWidth()
-			[
-				SNew(SEditableTextBox)
-				.Text(this, &FIKRigDefinitionDetails::GetGoalNameText, InItem)
-				.OnTextCommitted(this, &FIKRigDefinitionDetails::HandleGoalNameChanged, InItem)
-				.SelectAllTextWhenFocused(true)
-				.RevertTextOnEscape(true)
-				.Font(IDetailLayoutBuilder::GetDetailFont())
-			]
-		];
-}
-
-void FIKRigDefinitionDetails::HandleGoalNameChanged(const FText& NewName, ETextCommit::Type CommitType, FGoalNameListItemPtr InItem)
-{
-	//if (CommitType == ETextCommit::OnEnter)
-	if (IKRigController)
-	{
-		if (!NewName.IsEmptyOrWhitespace())
-		{
-			const FName NewFName = FName(*NewName.ToString());
-			if (InItem->DisplayName != NewFName)
-			{
-				IKRigController->RenameGoal(InItem->GoalName, NewFName);
-				InItem->GoalName = NewFName; // if you rename to the same as others, you'll reduce the number of goals
-				InItem->DisplayName = NewFName;
-
-				// refresh?
-
-			}
-		}
-	}
-}
-
-FText FIKRigDefinitionDetails::GetGoalNameText(FGoalNameListItemPtr InItem) const
-{
-	return FText::FromName(InItem->DisplayName);
-}
 #undef LOCTEXT_NAMESPACE
diff --git a/Engine/Plugins/Animation/IKRig/Source/IKRigEditor/Private/IKRigDefinitionDetails.h b/Engine/Plugins/Animation/IKRig/Source/IKRigEditor/Private/IKRigDefinitionDetails.h
--- a/Engine/Plugins/Animation/IKRig/Source/IKRigEditor/Private/IKRigDefinitionDetails.h
+++ b/Engine/Plugins/Animation/IKRig/Source/IKRigEditor/Private/IKRigDefinitionDetails.h
@@ -2,12 +2,8 @@
 
 #pragma once
 
-#include "CoreMinimal.h"
-#include "UObject/WeakObjectPtr.h"
 #include "UObject/WeakObjectPtrTemplates.h"
 #include "IDetailCustomization.h"
-#include "Types/SlateEnums.h"
-#include "Widgets/Views/SListView.h"
 #include "Input/Reply.h"
 
 class IDetailLayoutBuilder;
@@ -19,59 +15,37 @@ class ITableRow;
 class STableViewBase;
 class IPropertyHandle;
 
-// Utility class to build combo boxes out of arrays of names.
-struct FGoalNameListItem : public TSharedFromThis<FGoalNameListItem>
-{
-	FName GoalName;
-	FName DisplayName;
-
-	FGoalNameListItem(const FName& InName)
-	{
-		GoalName = InName;
-		DisplayName = InName;
-	}
-};
-
 class FIKRigDefinitionDetails : public IDetailCustomization
 {
+
 public:
-	/** Makes a new instance of this detail layout class for a specific detail view requesting it */
-	static TSharedRef<IDetailCustomization> MakeInstance();
 
 	virtual ~FIKRigDefinitionDetails();
 
+	/** Makes a new instance of this detail layout class for a specific detail view requesting it */
+	static TSharedRef<IDetailCustomization> MakeInstance();
+
 	/** IDetailCustomization interface */
 	virtual void CustomizeDetails( IDetailLayoutBuilder& DetailBuilder ) override;
 	virtual void CustomizeDetails(const TSharedPtr<IDetailLayoutBuilder>& DetailBuilder) override;
+	/** IDetailCustomization interface */
+
 private:
-	TWeakObjectPtr<UIKRigDefinition>		IKRigDefinition;
-	UIKRigController*						IKRigController;
 
-	// source asset
+	/** pointers to "Model" and "Controller" in MVC pattern */
+	TWeakObjectPtr<UIKRigDefinition> IKRigDefinition;
+	UIKRigController* IKRigController;
+
+	/** source skeleton asset, used to initialize hierarchy */
 	FString GetCurrentSourceAsset() const;
 	bool ShouldFilterAsset(const FAssetData& AssetData);
 	void OnAssetSelected(const FAssetData& AssetData);
-	TWeakObjectPtr<UObject>					SelectedAsset;
+	TWeakObjectPtr<UObject> SelectedAsset;
 
+	/** import new skeleton */
 	FReply OnImportHierarchy();
 	bool CanImport() const;
 
 	TWeakPtr< IDetailLayoutBuilder > DetailBuilderWeakPtr;
-
-	FReply OnShowSolverClassPicker();
-	FReply OnShowBoneSettingsClassPicker();
-
-	typedef TSharedPtr<FGoalNameListItem> FGoalNameListItemPtr;
-
-	void OnObjectPostEditChange(UObject* Object, FPropertyChangedEvent& InPropertyChangedEvent);
-	TSharedRef<ITableRow> OnGenerateWidgetForGoals(FGoalNameListItemPtr InItem, const TSharedRef<STableViewBase>& OwnerTable);
-	void HandleGoalNameChanged(const FText& NewName, ETextCommit::Type CommitType, FGoalNameListItemPtr InItem);
-	FText GetGoalNameText(FGoalNameListItemPtr InName) const;
-
-	TSharedPtr<IPropertyHandle> GoalPropertyHandle;
-	TSharedPtr<SListView<FGoalNameListItemPtr>> GoalListView;
-	TArray<FGoalNameListItemPtr> GoalListNames;
-
-	FDelegateHandle ObjectChangedDelegate;
 };
 
diff --git a/Engine/Plugins/Animation/IKRig/Source/IKRigEditor/Private/IKRigDefinitionFactory.cpp b/Engine/Plugins/Animation/IKRig/Source/IKRigEditor/Private/IKRigDefinitionFactory.cpp
--- a/Engine/Plugins/Animation/IKRig/Source/IKRigEditor/Private/IKRigDefinitionFactory.cpp
+++ b/Engine/Plugins/Animation/IKRig/Source/IKRigEditor/Private/IKRigDefinitionFactory.cpp
@@ -34,7 +34,7 @@ bool UIKRigDefinitionFactory::ConfigureProperties()
 
 FText UIKRigDefinitionFactory::GetDisplayName() const
 {
-	return LOCTEXT("IKRigDefinition_DisplayName", "IK Rig Definition");
+	return LOCTEXT("IKRigDefinition_DisplayName", "IK Rig");
 }
 
 uint32 UIKRigDefinitionFactory::GetMenuCategories() const
@@ -44,11 +44,11 @@ uint32 UIKRigDefinitionFactory::GetMenuCategories() const
 
 FText UIKRigDefinitionFactory::GetToolTip() const
 {
-	return LOCTEXT("IKRigDefinition_Tooltip", "Create IK Rig simply to animate or to use in runtime.");
+	return LOCTEXT("IKRigDefinition_Tooltip", "Create a customized IK Rig to be embedded in an Anim Graph or Control Rig.");
 }
 
 FString UIKRigDefinitionFactory::GetDefaultNewAssetName() const
 {
-	return FString(TEXT("NewIKRigDefinition"));
+	return FString(TEXT("NewIKRig"));
 }
 #undef LOCTEXT_NAMESPACE
diff --git a/Engine/Plugins/Animation/IKRig/Source/IKRigEditor/Private/IKRigEditor.cpp b/Engine/Plugins/Animation/IKRig/Source/IKRigEditor/Private/IKRigEditor.cpp
--- a/Engine/Plugins/Animation/IKRig/Source/IKRigEditor/Private/IKRigEditor.cpp
+++ b/Engine/Plugins/Animation/IKRig/Source/IKRigEditor/Private/IKRigEditor.cpp
@@ -15,7 +15,6 @@
 #include "AssetTypeActions_IKRigDefinition.h"
 #include "IKRigDefinition.h"
 #include "AnimMode/IKRigEditMode.h"
-#include "IKRigEffectorCustomization.h"
 
 IMPLEMENT_MODULE(FIKRigEditor, IKRigEditor)
 
@@ -31,8 +30,6 @@ void FIKRigEditor::StartupModule()
 	FPropertyEditorModule& PropertyModule = FModuleManager::GetModuleChecked<FPropertyEditorModule>("PropertyEditor");
 	PropertyModule.RegisterCustomClassLayout(UIKRigDefinition::StaticClass()->GetFName(), FOnGetDetailCustomizationInstance::CreateStatic(&FIKRigDefinitionDetails::MakeInstance));
 
-	PropertyModule.RegisterCustomPropertyTypeLayout("IKRigEffector", FOnGetPropertyTypeCustomizationInstance::CreateStatic(&FIKRigEffectorCustomization::MakeInstance));
-
 	FEditorModeRegistry::Get().RegisterMode<FIKRigEditMode>(UAnimGraphNode_IKRig::AnimModeName, LOCTEXT("IKRigEditMode", "IKRig"), FSlateIcon(), false);
 }
 
@@ -43,8 +40,6 @@ void FIKRigEditor::ShutdownModule()
 	FPropertyEditorModule& PropertyModule = FModuleManager::GetModuleChecked<FPropertyEditorModule>("PropertyEditor");
 	PropertyModule.UnregisterCustomPropertyTypeLayout("IKRigEffector");
 
-	PropertyModule.UnregisterCustomClassLayout(UIKRigDefinition::StaticClass()->GetFName());
-
 	if (IKRigDefinitionAssetAction.IsValid())
 	{
 		if (FModuleManager::Get().IsModuleLoaded("AssetTools"))
diff --git a/Engine/Plugins/Animation/IKRig/Source/IKRigEditor/Private/IKRigEffectorCustomization.cpp b/Engine/Plugins/Animation/IKRig/Source/IKRigEditor/Private/IKRigEffectorCustomization.cpp
deleted file mode 100644
diff --git a/Engine/Plugins/Animation/IKRig/Source/IKRigEditor/Private/IKRigEffectorCustomization.h b/Engine/Plugins/Animation/IKRig/Source/IKRigEditor/Private/IKRigEffectorCustomization.h
deleted file mode 100644
diff --git a/Engine/Plugins/Animation/IKRig/Source/IKRigEditor/Public/IKRigController.h b/Engine/Plugins/Animation/IKRig/Source/IKRigEditor/Public/IKRigController.h
--- a/Engine/Plugins/Animation/IKRig/Source/IKRigEditor/Public/IKRigController.h
+++ b/Engine/Plugins/Animation/IKRig/Source/IKRigEditor/Public/IKRigController.h
@@ -8,12 +8,6 @@
 #include "IKRigDefinition.h"
 #include "IKRigController.generated.h"
 
-// DELEGATE for solver chnage
-// DELEGATE for hierarchy changes
-// DELEGATE for profile changes
-// DELEGATE for constraint changes
-// DELEGATE for goals
-
 class UIKRigSolverDefinition;
 class UIKRigBoneSetting;
 struct FIKRigHierarchy;
@@ -28,90 +22,46 @@ class IKRIGEDITOR_API UIKRigController : public UObject
 
 public:
 
-	void SetSkeleton(const FReferenceSkeleton& InSkeleton);
-	const FIKRigHierarchy* GetHierarchy() const
-	{
-		return (IKRigDefinition)? &IKRigDefinition->Hierarchy : nullptr;
-	}
-
-	const TArray<FTransform>& GetRefPoseTransforms() const
-	{
-		if (IKRigDefinition)
-		{
-			return IKRigDefinition->RefPoseTransforms;
-		}
-
-		static TArray<FTransform> Dummy;
-		return Dummy;
-	}
+	//** use this to get a handle to a controller for the given IKRig */
+	static UIKRigController* GetControllerByRigDefinition(UIKRigDefinition* InIKRigDefinition);
 
+	//** skeleton */
+	const FIKRigHierarchy* GetHierarchy() const;
+	const TArray<FTransform>& GetRefPoseTransforms() const;
+	void SetSkeleton(const FReferenceSkeleton& InSkeleton);
 	bool AddBone(const FName& InName, const FName& InParent, const FTransform& InGlobalTransform);
-	bool RemoveBone(const FName& InName);
-	bool RenameBone(const FName& InOldName, const FName& InNewName);
-	bool ReparentBone(const FName& InName, const FName& InNewParent);
 	void ResetHierarchy();
 
-	// todo: more APIs for hierarchy operators
-
-	// solver operators
+	//** solvers */
 	UIKRigSolver* AddSolver(TSubclassOf<UIKRigSolver> InSolverClass);
 	void RemoveSolver(UIKRigSolver* SolverToDelete);
 	UIKRigSolver* GetSolver(int32 Index) const;
 	int32 GetNumSolvers() const;
 
-	// per-bone settings and constraints
-	UIKRigBoneSetting* AddBoneSetting(TSubclassOf<UIKRigBoneSetting> NewBoneSettingType);
-
-	// goal operators
-	void QueryGoals(TArray<FName>& OutGoals) const;
-	// why do we need rename? @todo: remove - maybe in the future when we want a tool that does rename automatically this could be useful, but for now it's not really
-	// goals are collected from solvers, so they could change name, and if it doesn't exists, it doesn't exist
-	void RenameGoal(const FName& OldName, const FName& NewName);
+	//** per-bone settings and constraints */
+	UIKRigBoneSetting* AddBoneSetting(TSubclassOf<UIKRigBoneSetting> NewBoneSettingType) const;
 
-	FName GetGoalName(UIKRigSolver* InSolver, const FIKRigEffector& InEffector);
-	void SetGoalName(UIKRigSolver* InSolver, const FIKRigEffector& InEffector, const FName& NewGoalName);
-
-	// this is to modify default value
-	FIKRigGoal* GetGoal(const FName& InGoalName);
-	const FIKRigGoal* GetGoal(const FName& InGoalName) const;
+	//** goals */
+	void RenameGoal(const FName& OldName, const FName& NewName) const;
+	void GetGoalNames(TArray<FName>& OutGoals) const;
+	DECLARE_MULTICAST_DELEGATE(FGoalModified);
+	FGoalModified OnGoalModified;
 
 	// BEGIN UObject
-	virtual void BeginDestroy() override;
 	static void AddReferencedObjects(UObject* InThis, FReferenceCollector& Collector);
 	// END UObject
 
-	// Controller getter
-	// Use this interface to get controller 
-	static UIKRigController* GetControllerByRigDefinition(UIKRigDefinition* InIKRigDefinition);
-
-	// delegates
-	DECLARE_MULTICAST_DELEGATE(FGoalModified);
-
-	FGoalModified OnGoalModified;
 private:
+
 	UPROPERTY(transient)
 	UIKRigDefinition* IKRigDefinition = nullptr;
 
-	bool ValidateSolver(UIKRigSolver* const Solver) const;
-	void UpdateGoal();
-
-	FIKRigConstraintProfile* GetConstraintProfile(const FName& InProfileName) const;
-
-	void InitializeSolver(UIKRigSolver* Solver);
-	void UninitializeSolver(UIKRigSolver* Solver);
-
-	TMap<UIKRigSolver*, FDelegateHandle> SolverDelegateHandles;
-
 	static TMap<UIKRigDefinition*, UIKRigController*> DefinitionToControllerMap;
-
 	void SetIKRigDefinition(UIKRigDefinition* InIKRigDefinition);
 
-	void EnsureUniqueConstraintName(FName& InOutName);
-
-	// called by IKRigDefinition when it's begin destroy
+	//** called by IKRigDefinition when it's begin destroy */
 	static void RemoveControllerByRigDefinition(UIKRigDefinition* InIKRigDefinition);
 
-	// friend class 
 	friend class UIKRigDefinition;
 };
 
diff --git a/Engine/Plugins/Experimental/FullBodyIK/Source/FullBodyIK/Private/FullBodyIKSolver.cpp b/Engine/Plugins/Experimental/FullBodyIK/Source/FullBodyIK/Private/FullBodyIKSolver.cpp
deleted file mode 100644
diff --git a/Engine/Plugins/Experimental/FullBodyIK/Source/FullBodyIK/Private/FullBodyIKSolver.h b/Engine/Plugins/Experimental/FullBodyIK/Source/FullBodyIK/Private/FullBodyIKSolver.h
deleted file mode 100644
diff --git a/Engine/Plugins/Experimental/FullBodyIK/Source/FullBodyIK/Private/IKRigFBIKSolver.cpp b/Engine/Plugins/Experimental/FullBodyIK/Source/FullBodyIK/Private/IKRigFBIKSolver.cpp
new file mode 100644
--- /dev/null
+++ b/Engine/Plugins/Experimental/FullBodyIK/Source/FullBodyIK/Private/IKRigFBIKSolver.cpp
@@ -0,0 +1,566 @@
+// Copyright Epic Games, Inc. All Rights Reserved.
+
+/*=============================================================================
+	FullBodyIKSolver.cpp: Solver execution class for Transform
+=============================================================================*/
+
+#include "IKRigFBIKSolver.h"
+#include "IKRigDataTypes.h"
+#include "IKRigHierarchy.h"
+#include "FBIKConstraint.h"
+#include "JacobianIK.h"
+#include "FBIKUtil.h"
+#include "FBIKShared.h"
+#include "Drawing/ControlRigDrawInterface.h"
+
+const FName UIKRigFBIKSolver::EffectorTargetPrefix = FName(TEXT("FullBodyIKTarget"));
+
+//////////////////////////////////////////////
+// utility functions
+////////////////////////////////////////////////
+static float EnsureToAddBoneToLinkData(const FIKRigTransforms& TransformModifier, const int32 CurrentItem, TArray<FFBIKLinkData>& LinkData,
+	TMap<int32, int32>& HierarchyToLinkDataMap, TMap<int32, int32>& LinkDataToHierarchyIndices)
+{
+	float ChainLength = 0;
+
+	// we insert from back to first, but only if the list doesn't have it
+	int32* FoundLinkIndex = HierarchyToLinkDataMap.Find(CurrentItem);
+	if (!FoundLinkIndex)
+	{
+		int32 NewLinkIndex = LinkData.AddDefaulted();
+		FFBIKLinkData& NewLink = LinkData[NewLinkIndex];
+
+		// find parent LinkIndex
+		const int32 ParentItem = TransformModifier.Hierarchy->GetParentIndex(CurrentItem);
+		FoundLinkIndex = HierarchyToLinkDataMap.Find(ParentItem);
+		NewLink.ParentLinkIndex = (FoundLinkIndex) ? *FoundLinkIndex : INDEX_NONE;
+		NewLink.SetTransform(TransformModifier.GetGlobalTransform(CurrentItem));
+
+		if (ParentItem != INDEX_NONE)
+		{
+			// @todo: currently we assume the input pose is init transform
+			// this could use reference pose getter from IKRigSolver or we could wrap it to TransformModifier
+			FVector DiffLocation = TransformModifier.GetGlobalTransform(CurrentItem).GetLocation() - TransformModifier.GetGlobalTransform(ParentItem).GetLocation();
+			// set Length
+			NewLink.Length = DiffLocation.Size();
+		}
+
+		// we create bidirectional look up table
+		HierarchyToLinkDataMap.Add(CurrentItem, NewLinkIndex);
+		LinkDataToHierarchyIndices.Add(NewLinkIndex, CurrentItem);
+
+		ChainLength = NewLink.Length;
+	}
+	else
+	{
+		ChainLength = LinkData[*FoundLinkIndex].Length;
+	}
+
+	return ChainLength;
+}
+
+static void AddToEffectorTarget(int32 EffectorIndex, const int32 Effector, TMap<int32, FFBIKEffectorTarget>& EffectorTargets, const TMap<int32, int32>& HierarchyToLinkDataMap,
+	TArray<int32>& EffectorLinkIndices, float ChainLength, const TArray<int32>& ChainIndices, int32 PositionDepth, int32 RotationDepth)
+{
+	const int32* EffectorLinkIndexID = HierarchyToLinkDataMap.Find(Effector);
+	check(EffectorLinkIndexID);
+	EffectorLinkIndices[EffectorIndex] = *EffectorLinkIndexID;
+	// add EffectorTarget for this link Index
+	FFBIKEffectorTarget& EffectorTarget = EffectorTargets.FindOrAdd(*EffectorLinkIndexID);
+	EffectorTarget.ChainLength = ChainLength;
+
+	// convert bone chain indices to link chain
+	const int32 MaxNum = FMath::Max(PositionDepth, RotationDepth);
+	if (ensure(MaxNum <= ChainIndices.Num()))
+	{
+		EffectorTarget.LinkChain.Reset(MaxNum);
+		for (int32 Index = 0; Index < MaxNum; ++Index)
+		{
+			const int32 Bone = ChainIndices[Index];
+			const int32* LinkIndexID = HierarchyToLinkDataMap.Find(Bone);
+			EffectorTarget.LinkChain.Add(*LinkIndexID);
+		}
+	}
+}
+
+// the output should be from current index -> to root parent 
+static bool GetBoneChain(const FIKRigTransforms& TransformModifier, const FName& Root, const FName& Current, TArray<int32>& ChainIndices)
+{
+	ChainIndices.Reset();
+
+	int32 RootIndex = TransformModifier.Hierarchy->GetIndex(Root);
+	int32 Iterator = TransformModifier.Hierarchy->GetIndex(Current);;
+
+	// iterates until key is valid
+	while (Iterator!=INDEX_NONE && Iterator != RootIndex)
+	{
+		ChainIndices.Insert(Iterator, 0);
+		Iterator = TransformModifier.Hierarchy->GetParentIndex(Iterator);
+	}
+
+	// add the last one if valid
+	if (Iterator!=INDEX_NONE)
+	{
+		ChainIndices.Insert(Iterator, 0);
+	}
+
+	// when you reached to something invalid, that means, we did not hit the 
+	// expected root, we iterated to the root and above, and we exhausted our option
+	// so if you hit or valid target by the hitting max depth, we should 
+	return !ChainIndices.IsEmpty();
+}
+
+static void AddEffectors(const FIKRigTransforms& TransformModifier, const FName Root, const TArray<FFBIKRigEffector>& Effectors,
+	TArray<FFBIKLinkData>& LinkData, TMap<int32, FFBIKEffectorTarget>& EffectorTargets, TArray<int32>& EffectorLinkIndices,
+	TMap<int32, int32>& LinkDataToHierarchyIndices, TMap<int32, int32>& HierarchyToLinkDataMap, const FSolverInput& SolverProperty)
+{
+	EffectorLinkIndices.SetNum(Effectors.Num());
+	// fill up all effector indices
+	for (int32 Index = 0; Index < Effectors.Num(); ++Index)
+	{
+		// clear link indices, so that we don't search
+		EffectorLinkIndices[Index] = INDEX_NONE;
+		// create LinkeData from root bone to all effectors 
+		const FName Bone = Effectors[Index].Target.Bone;
+		const int32 BoneIndex = TransformModifier.Hierarchy->GetIndex(Bone);
+
+		if (Bone == NAME_None || BoneIndex == INDEX_NONE)
+		{
+			continue;
+		}
+
+		const FFBIKRigEffector& CurrentEffector = Effectors[Index];
+
+		TArray<int32> ChainIndices;
+		// if we haven't got to root, this is not valid chain
+		if (GetBoneChain(TransformModifier, Root, Bone, ChainIndices))
+		{
+			auto CalculateStrength = [&](int32 InBoneChainDepth, const int32 MaxDepth, float CurrentStrength, float MinStrength) -> float
+			{
+				const float Range = FMath::Max(CurrentStrength - MinStrength, 0.f);
+				const float ApplicationStrength = (float)(1.f - (float)InBoneChainDepth / (float)MaxDepth) * Range;
+				return ApplicationStrength + MinStrength;
+				//return FMath::Clamp(ApplicationStrength + MinStrength, MinStrength, CurrentStrength);
+			};
+
+			auto UpdateMotionStrength = [&](int32 InBoneChainDepth,
+				const int32 MaxPositionDepth, const int32 MaxRotationDepth, FFBIKLinkData& InOutNewLink)
+			{
+				// add motion scales
+				float LinearMotionStrength;
+				float AngularMotionStrength;
+
+				if (CurrentEffector.PositionDepth <= InBoneChainDepth)
+				{
+					LinearMotionStrength = 0.f;
+				}
+				else
+				{
+					LinearMotionStrength = CalculateStrength(InBoneChainDepth, MaxPositionDepth, SolverProperty.LinearMotionStrength, SolverProperty.MinLinearMotionStrength);
+				}
+
+				if (CurrentEffector.RotationDepth <= InBoneChainDepth)
+				{
+					AngularMotionStrength = 0.f;
+				}
+				else
+				{
+					AngularMotionStrength = CalculateStrength(InBoneChainDepth, MaxRotationDepth, SolverProperty.AngularMotionStrength, SolverProperty.MinAngularMotionStrength);
+				}
+
+				InOutNewLink.AddMotionStrength(LinearMotionStrength, AngularMotionStrength);
+			};
+
+			// position depth and rotation depth can't go beyond of it
+			// for now we cull it. 
+			const int32 PositionDepth = FMath::Min(CurrentEffector.PositionDepth, ChainIndices.Num());
+			const int32 RotationDepth = FMath::Min(CurrentEffector.RotationDepth, ChainIndices.Num());
+
+			float ChainLength = 0.f;
+
+			// add to link data
+			for (int32 BoneChainIndex = 0; BoneChainIndex < ChainIndices.Num(); ++BoneChainIndex)
+			{
+				const int32 CurrentItem = ChainIndices[BoneChainIndex];
+				ChainLength += EnsureToAddBoneToLinkData(TransformModifier, CurrentItem, LinkData,
+					HierarchyToLinkDataMap, LinkDataToHierarchyIndices);
+
+				const int32 ChainDepth = ChainIndices.Num() - BoneChainIndex;
+				int32* FoundLinkIndex = HierarchyToLinkDataMap.Find(CurrentItem);
+
+				// now we should always have it
+				check(FoundLinkIndex);
+				UpdateMotionStrength(ChainDepth, PositionDepth, RotationDepth, LinkData[*FoundLinkIndex]);
+			}
+
+			// add to EffectorTargets
+			AddToEffectorTarget(Index, BoneIndex, EffectorTargets, HierarchyToLinkDataMap, EffectorLinkIndices, ChainLength, ChainIndices, PositionDepth, RotationDepth);
+		}
+	}
+}
+/////////////////////////////////////////////////
+
+UIKRigFBIKSolver::UIKRigFBIKSolver()
+{
+}
+
+void UIKRigFBIKSolver::Init(const FIKRigTransforms& InGlobalTransform)
+{
+
+	LinkData.Reset();
+	EffectorTargets.Reset();
+	EffectorLinkIndices.Reset();
+	LinkDataToHierarchyIndices.Reset();
+	HierarchyToLinkDataMap.Reset();
+
+	// verify the chain
+	AddEffectors(InGlobalTransform, Root, Effectors, LinkData, EffectorTargets, EffectorLinkIndices, LinkDataToHierarchyIndices, HierarchyToLinkDataMap, SolverProperty);
+	
+}
+
+void UIKRigFBIKSolver::Solve(
+	FIKRigTransforms& InOutGlobalTransform,
+	const FIKRigGoalContainer& Goals,
+	FControlRigDrawInterface* InOutDrawInterface)
+{
+	if (Effectors.IsEmpty())
+	{
+		return; // nothing to do
+	}
+	
+	if (!LinkDataToHierarchyIndices.IsEmpty())
+	{
+#if 0
+		// we do this every frame for now
+		if (Constraints.Num() > 0)
+		{
+			DECLARE_SCOPE_HIERARCHICAL_COUNTER(TEXT("Build Constraint"))
+			//Build constraints
+			FBIKConstraintLib::BuildConstraints(Constraints, InternalConstraints, InOutGlobalTransform, LinkData, LinkDataToHierarchyIndices, HierarchyToLinkDataMap);
+		}
+
+		// during only editor and update
+		// we expect solver type changes, it will reinit
+		// InternalConstraints can't be changed during runtime
+		if (!InternalConstraints.IsEmpty())
+		{
+			WorkData.IKSolver.SetPostProcessDelegateForIteration(FPostProcessDelegateForIteration::CreateStatic(&FBIKConstraintLib::ApplyConstraint, &InternalConstraints));
+		}
+		else
+		{
+			WorkData.IKSolver.ClearPostProcessDelegateForIteration();
+		}
+	//disable constraint for now
+#endif // 0 
+
+		// before update we finalize motion scale
+		// this code may go away once we have constraint
+		// update link data and end effectors
+		for (int32 LinkIndex = 0; LinkIndex < LinkData.Num(); ++LinkIndex)
+		{
+			const int32 Item = *LinkDataToHierarchyIndices.Find(LinkIndex);
+			LinkData[LinkIndex].SetTransform(InOutGlobalTransform.GetGlobalTransform(Item));
+			// @todo: fix this somewhere else - we can add this to prepare step
+			// @todo: we update motion scale here, then?
+			LinkData[LinkIndex].FinalizeForSolver();
+		}
+
+		// update mid effector info
+		const float LinearMotionStrength = FMath::Max(SolverProperty.LinearMotionStrength, SolverProperty.MinLinearMotionStrength);
+		const float AngularMotionStrength = FMath::Max(SolverProperty.AngularMotionStrength, SolverProperty.MinAngularMotionStrength);
+		const float LinearRange = LinearMotionStrength - SolverProperty.MinLinearMotionStrength;
+		const float AngularRange = AngularMotionStrength - SolverProperty.MinAngularMotionStrength;
+
+		//const TArray<FFBIKRigEffector>& Effectors = Effectors;
+
+		// update end effector info
+		for (int32 EffectorIndex = 0; EffectorIndex < Effectors.Num(); ++EffectorIndex)
+		{
+			int32 EffectorLinkIndex = EffectorLinkIndices[EffectorIndex];
+			if (EffectorLinkIndex != INDEX_NONE)
+			{
+				FFBIKEffectorTarget* EffectorTarget = EffectorTargets.Find(EffectorLinkIndex);
+				if (EffectorTarget)
+				{
+					const FFBIKRigEffector& CurEffector = Effectors[EffectorIndex];
+					const FVector CurrentLinkLocation = LinkData[EffectorLinkIndex].GetTransform().GetLocation();
+					const FQuat CurrentLinkRotation = LinkData[EffectorLinkIndex].GetTransform().GetRotation();
+					FIKRigGoal Goal;
+					ensure(GetGoalForEffector(CurEffector.Target, Goals, Goal));
+					EffectorTarget->Position = Goal.Position; // FMath::Lerp(CurrentLinkLocation, EffectorLocation, RigTarget.PositionTarget.);
+					EffectorTarget->Rotation = Goal.Rotation; // FMath::Lerp(CurrentLinkRotation, EffectorRotation, CurEffector.RotationAlpha);
+					EffectorTarget->InitialPositionDistance = (EffectorTarget->Position - CurrentLinkLocation).Size();
+					EffectorTarget->InitialRotationDistance = (FBIKUtil::GetScaledRotationAxis(EffectorTarget->Rotation) - FBIKUtil::GetScaledRotationAxis(CurrentLinkRotation)).Size();
+
+					const float Pull = FMath::Clamp(CurEffector.Pull, 0.f, 1.f);
+					// we want some impact of Pull, in order for Pull to have some impact, we clamp to some number
+					const float TargetClamp = FMath::Clamp(SolverProperty.DefaultTargetClamp, 0.f, 0.7f);
+					const float Scale = TargetClamp + Pull * (1.f - TargetClamp);
+					// Pull set up
+					EffectorTarget->LinearMotionStrength = LinearRange * Scale + SolverProperty.MinLinearMotionStrength;
+					EffectorTarget->AngularMotionStrength = AngularRange * Scale + SolverProperty.MinAngularMotionStrength;
+					EffectorTarget->ConvergeScale = Scale;
+					EffectorTarget->TargetClampScale = Scale;
+
+					EffectorTarget->bPositionEnabled = true;
+					EffectorTarget->bRotationEnabled = true;
+				}
+			}
+		}
+
+		DebugData.Reset();
+
+		const bool bDebugEnabled = DebugOption.bDrawDebugHierarchy || DebugOption.bDrawDebugEffector || DebugOption.bDrawDebugConstraints;
+
+		// we can't reuse memory until we fix the memory issue on RigVM
+		{
+			IKSolver.SolveJacobianIK(LinkData, EffectorTargets,
+				JacobianIK::FSolverParameter(SolverProperty.Damping, true, false, (SolverProperty.bUseJacobianTranspose) ? EJacobianSolver::JacobianTranspose : EJacobianSolver::JacobianPIDLS),
+				SolverProperty.MaxIterations, SolverProperty.Precision, &DebugData);
+
+			if (MotionProperty.bForceEffectorRotationTarget)
+			{
+				// if position is reached, we force rotation target
+				for (int32 EffectorIndex = 0; EffectorIndex < Effectors.Num(); ++EffectorIndex)
+				{
+					int32 EffectorLinkIndex = EffectorLinkIndices[EffectorIndex];
+					if (EffectorLinkIndex != INDEX_NONE)
+					{
+						FFBIKEffectorTarget* EffectorTarget = EffectorTargets.Find(EffectorLinkIndex);
+						if (EffectorTarget && EffectorTarget->bRotationEnabled)
+						{
+							bool bApplyRotation = true;
+
+							if (MotionProperty.bOnlyApplyWhenReachedToTarget)
+							{
+								// only do this when position is reached? This will conflict with converge scale
+								const FVector& BonePosition = LinkData[EffectorLinkIndex].GetTransform().GetLocation();
+								const FVector& TargetPosition = EffectorTarget->Position;
+
+								bApplyRotation = (FVector(BonePosition - TargetPosition).SizeSquared() <= SolverProperty.Precision * SolverProperty.Precision);
+							}
+
+							if (bApplyRotation)
+							{
+								FQuat NewRotation = EffectorTarget->Rotation;
+								FTransform NewTransform = LinkData[EffectorLinkIndex].GetTransform();
+								NewTransform.SetRotation(NewRotation);
+								LinkData[EffectorLinkIndex].SetTransform(NewTransform);
+							}
+						}
+					}
+				}
+			}
+		}
+		///////////////////////////////////////////////////////////////////////////
+		// debug draw start
+		///////////////////////////////////////////////////////////////////////////
+		if (bDebugEnabled && InOutDrawInterface != nullptr)
+		{
+			const int32 DebugDataNum = DebugData.Num();
+			if (!DebugData.IsEmpty())
+			{
+				for (int32 DebugIndex = DebugDataNum - 1; DebugIndex >= 0; --DebugIndex)
+				{
+					const TArray<FFBIKLinkData>& LocalLink = DebugData[DebugIndex].LinkData;
+
+					FTransform Offset = DebugOption.DrawWorldOffset;
+					Offset.SetLocation(Offset.GetLocation() * (DebugDataNum - DebugIndex));
+
+					if (DebugOption.bDrawDebugHierarchy)
+					{
+						for (int32 LinkIndex = 0; LinkIndex < LocalLink.Num(); ++LinkIndex)
+						{
+							const FFBIKLinkData& Data = LocalLink[LinkIndex];
+
+							FLinearColor DrawColor = FLinearColor::White;
+
+							float LineThickness = 0.f;
+							if (DebugOption.bColorAngularMotionStrength || DebugOption.bColorLinearMotionStrength)
+							{
+								DrawColor = FLinearColor::Black;
+								if (DebugOption.bColorAngularMotionStrength)
+								{
+									const float Range = FMath::Max(SolverProperty.AngularMotionStrength - SolverProperty.MinAngularMotionStrength, 0.f);
+									if (Range > 0.f)
+									{
+										float CurrentStrength = Data.GetAngularMotionStrength() - SolverProperty.MinAngularMotionStrength;
+										float Alpha = FMath::Clamp(CurrentStrength / Range, 0.f, 1.f);
+										DrawColor.R = LineThickness = Alpha;
+									}
+								}
+								else if (DebugOption.bColorLinearMotionStrength)
+								{
+									const float Range = FMath::Max(SolverProperty.LinearMotionStrength - SolverProperty.MinLinearMotionStrength, 0.f);
+									if (Range > 0.f)
+									{
+										float CurrentStrength = Data.GetLinearMotionStrength() - SolverProperty.MinLinearMotionStrength;
+										float Alpha = FMath::Clamp(CurrentStrength / Range, 0.f, 1.f);
+										DrawColor.B = LineThickness = Alpha;
+									}
+								}
+							}
+
+							if (Data.ParentLinkIndex != INDEX_NONE)
+							{
+								const FFBIKLinkData& ParentData = LocalLink[Data.ParentLinkIndex];
+								InOutDrawInterface->DrawLine(Offset, Data.GetPreviousTransform().GetLocation(), ParentData.GetPreviousTransform().GetLocation(), DrawColor, LineThickness);
+							}
+
+							if (DebugOption.bDrawDebugAxes)
+							{
+								InOutDrawInterface->DrawAxes(Offset, Data.GetPreviousTransform(), DebugOption.DrawSize);
+							}
+						}
+					}
+
+					if (DebugOption.bDrawDebugEffector)
+					{
+						for (auto Iter = EffectorTargets.CreateConstIterator(); Iter; ++Iter)
+						{
+							const FFBIKEffectorTarget& EffectorTarget = Iter.Value();
+							if (EffectorTarget.bPositionEnabled)
+							{
+								// draw effector target locations
+								InOutDrawInterface->DrawBox(Offset, FTransform(EffectorTarget.Position), FLinearColor::Yellow, DebugOption.DrawSize);
+							}
+
+							// draw effector link location
+							InOutDrawInterface->DrawBox(Offset, LocalLink[Iter.Key()].GetPreviousTransform(), FLinearColor::Green, DebugOption.DrawSize);
+						}
+
+						for (int32 Index = 0; Index < DebugData[DebugIndex].TargetVectorSources.Num(); ++Index)
+						{
+							// draw arrow to the target
+							InOutDrawInterface->DrawLine(Offset, DebugData[DebugIndex].TargetVectorSources[Index].GetLocation(),
+								DebugData[DebugIndex].TargetVectorSources[Index].GetLocation() + DebugData[DebugIndex].TargetVectors[Index], FLinearColor::Red);
+						}
+					}
+				}
+			}
+#if 0 
+			if (SolverDef->DebugOption.bDrawDebugConstraints && InternalConstraints.Num())
+			{
+				FTransform Offset = FTransform::Identity;
+
+				// draw frame if active
+				for (int32 Index = 0; Index < Constraints.Num(); ++Index)
+				{
+					if (Constraints[Index].bEnabled)
+					{
+						if (Constraints[Index].Item.IsValid())
+						{
+							const int32* Found = HierarchyToLinkDataMap.Find(Constraints[Index].Item);
+							if (Found)
+							{
+								FTransform ConstraintFrame = LinkData[*Found].GetTransform();
+								ConstraintFrame.ConcatenateRotation(FQuat(Constraints[Index].OffsetRotation));
+								InOutDrawInterface->DrawAxes(Offset, ConstraintFrame, 2.f);
+							}
+						}
+					}
+				}
+
+				for (int32 Index = 0; Index < InternalConstraints.Num(); ++Index)
+				{
+					// for now we have rotation limit only
+					if (InternalConstraints[Index].IsType< FRotationLimitConstraint>())
+					{
+						FRotationLimitConstraint& LimitConstraint = InternalConstraints[Index].Get<FRotationLimitConstraint>();
+						const FQuat LocalRefRotation = LimitConstraint.RelativelRefPose.GetRotation();
+						FTransform RotationTransform = LinkData[LimitConstraint.ConstrainedIndex].GetTransform();
+						// base is parent transform but in their space, we can get there by inversing local ref rotation
+						FTransform BaseTransform = FTransform(LocalRefRotation).GetRelativeTransformReverse(RotationTransform);
+						BaseTransform.ConcatenateRotation(LimitConstraint.BaseFrameOffset);
+						InOutDrawInterface->DrawAxes(Offset, BaseTransform, 5.f, 1.f);
+
+						// current transform
+						const FQuat LocalRotation = BaseTransform.GetRotation().Inverse() * RotationTransform.GetRotation();
+						const FQuat DeltaTransform = LocalRefRotation.Inverse() * LocalRotation;
+						RotationTransform.SetRotation(BaseTransform.GetRotation() * DeltaTransform);
+						RotationTransform.NormalizeRotation();
+						RotationTransform.SetLocation(BaseTransform.GetLocation());
+
+						// draw ref pose on their current transform
+						InOutDrawInterface->DrawAxes(Offset, RotationTransform, 10.f, 1.f);
+
+						FVector XAxis = BaseTransform.GetUnitAxis(EAxis::X);
+						FVector YAxis = BaseTransform.GetUnitAxis(EAxis::Y);
+						FVector ZAxis = BaseTransform.GetUnitAxis(EAxis::Z);
+
+						if (LimitConstraint.bXLimitSet)
+						{
+							FTransform XAxisConeTM(YAxis, XAxis ^ YAxis, XAxis, BaseTransform.GetTranslation());
+							XAxisConeTM.SetRotation(FQuat(XAxis, 0.f) * XAxisConeTM.GetRotation());
+							XAxisConeTM.SetScale3D(FVector(30.f));
+							InOutDrawInterface->DrawCone(Offset, XAxisConeTM, LimitConstraint.Limit.X, 0.0f, 24, false, FLinearColor::Red, GEngine->ConstraintLimitMaterialX->GetRenderProxy());
+						}
+
+						if (LimitConstraint.bYLimitSet)
+						{
+							FTransform YAxisConeTM(ZAxis, YAxis ^ ZAxis, YAxis, BaseTransform.GetTranslation());
+							YAxisConeTM.SetRotation(FQuat(YAxis, 0.f) * YAxisConeTM.GetRotation());
+							YAxisConeTM.SetScale3D(FVector(30.f));
+							InOutDrawInterface->DrawCone(Offset, YAxisConeTM, LimitConstraint.Limit.Y, 0.0f, 24, false, FLinearColor::Green, GEngine->ConstraintLimitMaterialY->GetRenderProxy());
+						}
+
+						if (LimitConstraint.bZLimitSet)
+						{
+							FTransform ZAxisConeTM(XAxis, ZAxis ^ XAxis, ZAxis, BaseTransform.GetTranslation());
+							ZAxisConeTM.SetRotation(FQuat(ZAxis, 0.f) * ZAxisConeTM.GetRotation());
+							ZAxisConeTM.SetScale3D(FVector(30.f));
+							InOutDrawInterface->DrawCone(Offset, ZAxisConeTM, LimitConstraint.Limit.Z, 0.0f, 24, false, FLinearColor::Blue, GEngine->ConstraintLimitMaterialZ->GetRenderProxy());
+						}
+					}
+
+					if (InternalConstraints[Index].IsType<FPoleVectorConstraint>())
+					{
+						// darw pole vector location
+						// draw 3 joints line and a plane to pole vector
+						FPoleVectorConstraint& Constraint = InternalConstraints[Index].Get<FPoleVectorConstraint>();
+						FTransform RootTransform = LinkData[Constraint.ParentBoneIndex].GetTransform();
+						FTransform JointTransform = LinkData[Constraint.BoneIndex].GetTransform();
+						FTransform ChildTransform = LinkData[Constraint.ChildBoneIndex].GetTransform();
+
+						FVector JointTarget = (Constraint.bUseLocalDir) ? Constraint.CalculateCurrentPoleVectorDir(RootTransform, JointTransform, ChildTransform, LinkData[Constraint.BoneIndex].LocalFrame) : Constraint.PoleVector;
+
+						// draw the plane, 
+						TArray<FVector> Positions;
+						Positions.Add(RootTransform.GetLocation());
+						Positions.Add(ChildTransform.GetLocation());
+						Positions.Add(ChildTransform.GetLocation());
+						Positions.Add(JointTarget);
+						Positions.Add(JointTarget);
+						Positions.Add(RootTransform.GetLocation());
+
+						InOutDrawInterface->DrawLines(Offset, Positions, FLinearColor::Gray, 1.2f);
+						InOutDrawInterface->DrawLine(Offset, JointTransform.GetLocation(), JointTarget, FLinearColor::Red, 1.2f);
+					}
+				}
+			}
+
+			#endif // 0 - disable all constraint
+		}
+		///////////////////////////////////////////////////////////////////////////
+		// debug draw end
+		///////////////////////////////////////////////////////////////////////////
+	}
+
+	// we update back to hierarchy
+	for (int32 LinkIndex = 0; LinkIndex < LinkData.Num(); ++LinkIndex)
+	{
+		// only propagate, if you are leaf joints here
+		// this means, only the last joint in the test
+		const int32 CurrentItem = *LinkDataToHierarchyIndices.Find(LinkIndex);
+		const FTransform& LinkTransform = LinkData[LinkIndex].GetTransform();
+		InOutGlobalTransform.SetGlobalTransform(CurrentItem, LinkTransform, true);
+	}
+}
+
+void UIKRigFBIKSolver::CollectGoalNames(TSet<FName>& OutGoals) const
+{
+	for (const FFBIKRigEffector& Effector: Effectors)
+	{
+		OutGoals.Add(Effector.Target.Goal);
+	}
+}
diff --git a/Engine/Plugins/Experimental/FullBodyIK/Source/FullBodyIK/Private/IKRigFBIKSolver.h b/Engine/Plugins/Experimental/FullBodyIK/Source/FullBodyIK/Private/IKRigFBIKSolver.h
new file mode 100644
--- /dev/null
+++ b/Engine/Plugins/Experimental/FullBodyIK/Source/FullBodyIK/Private/IKRigFBIKSolver.h
@@ -0,0 +1,101 @@
+// Copyright Epic Games, Inc. All Rights Reserved.
+
+/**
+ * Contains Transform Solver Execution
+ *
+ */
+
+#pragma once
+
+#include "IKRigSolver.h"
+#include "FBIKShared.h"
+#include "FBIKConstraintOption.h"
+#include "FBIKDebugOption.h"
+#include "IKRigFBIKSolver.generated.h"
+
+struct FFBIKLinkData;
+
+USTRUCT()
+struct FFBIKRigEffector
+{
+	GENERATED_BODY()
+
+	/**
+	* The last item in the chain to solve - the effector
+	*/
+	UPROPERTY(EditAnywhere, Category = FFBIKRigEffector)
+	FIKRigEffectorGoal Target;
+
+	UPROPERTY(EditAnywhere, Category = FFBIKRigEffector)
+	int32  PositionDepth = 1000;
+
+	UPROPERTY(EditAnywhere, Category = FFBIKRigEffector)
+	int32  RotationDepth = 1000;
+
+	/*
+	 * Clamps the total length to target by this scale for each iteration
+	 * This helps to stabilize solver to reduce singularity by avoiding to try to reach target too far.
+	 */
+	UPROPERTY(EditAnywhere, Category = FFBIKRigEffector)
+	float Pull = 0.f;
+
+	FFBIKRigEffector()
+	{
+	}
+};
+
+ // run time for UFullBodyIKSolverDefinition
+UCLASS(EditInlineNew)
+class FULLBODYIK_API UIKRigFBIKSolver : public UIKRigSolver
+{
+	GENERATED_BODY()
+
+public:
+	
+	UPROPERTY(EditAnywhere, Category = Effectors)
+	FName Root;
+
+	UPROPERTY(EditAnywhere, Category = Effectors)
+	TArray<FFBIKRigEffector> Effectors;
+
+	UPROPERTY(EditAnywhere, Category = Solver)
+	FSolverInput SolverProperty;
+
+	UPROPERTY(EditAnywhere, Category = Motion)
+	FMotionProcessInput MotionProperty;
+
+	UPROPERTY(EditAnywhere, Category = Solver)
+	FFBIKDebugOption DebugOption;
+
+	UIKRigFBIKSolver();
+
+private:
+
+	/** list of Link Data for solvers - joints */
+	TArray<FFBIKLinkData> LinkData;
+	/** Effector Targets - search key is LinkData Index */
+	TMap<int32, FFBIKEffectorTarget> EffectorTargets; 
+	/** End Effector Link Indices - EndEffector index to LinkData index*/
+	TArray<int32> EffectorLinkIndices;
+	/** Map from LinkData index to Hierarchy Index*/
+	TMap<int32, int32> LinkDataToHierarchyIndices;
+	/** Map from Rig Hierarchy Index to LinkData index*/
+	TMap<int32, int32> HierarchyToLinkDataMap;
+	/** Constraints data */
+	TArray<ConstraintType> InternalConstraints;
+	/* Current Solver */
+	FJacobianSolver_FullbodyIK IKSolver;
+	/** Debug Data */
+	TArray<FJacobianDebugData> DebugData;
+
+protected:
+	virtual void Init(const FIKRigTransforms& InGlobalTransforms) override;
+	virtual void Solve(
+		FIKRigTransforms& InOutGlobalTransforms,
+		const FIKRigGoalContainer& Goals,
+		FControlRigDrawInterface* InOutDrawInterface) override;
+	virtual void CollectGoalNames(TSet<FName>& OutGoals) const override;
+
+public:
+	static const FName EffectorTargetPrefix;
+};
\ No newline at end of file
