Removing unnecessary dependency in PBIK module of FBIK plugin.

diff --git a/Engine/Plugins/Experimental/FullBodyIK/Source/PBIK/PBIK.Build.cs b/Engine/Plugins/Experimental/FullBodyIK/Source/PBIK/PBIK.Build.cs
--- a/Engine/Plugins/Experimental/FullBodyIK/Source/PBIK/PBIK.Build.cs
+++ b/Engine/Plugins/Experimental/FullBodyIK/Source/PBIK/PBIK.Build.cs
@@ -13,8 +13,7 @@ public class PBIK : ModuleRules
 			{
 				"Core",
 				"RigVM",
-				"ControlRig",
-				"IKRig"
+				"ControlRig"
 			}
 			);
 
