Fix non-initialized struct members
Fixed bug, parent indices not always correct in new Control Rig hierarchy
Minor cleanup, const correctness in PBIK

diff --git a/Engine/Plugins/Experimental/FullBodyIK/Source/PBIK/Private/Core/PBIKBody.cpp b/Engine/Plugins/Experimental/FullBodyIK/Source/PBIK/Private/Core/PBIKBody.cpp
--- a/Engine/Plugins/Experimental/FullBodyIK/Source/PBIK/Private/Core/PBIKBody.cpp
+++ b/Engine/Plugins/Experimental/FullBodyIK/Source/PBIK/Private/Core/PBIKBody.cpp
@@ -11,8 +11,7 @@ FBone::FBone(
 	const int& InParentIndex,		// must pass -1 for root of whole skeleton
 	const FVector& InOrigPosition,
 	const FQuat& InOrigRotation,
-	bool bInIsSolverRoot,
-	const int& InElementIndex)
+	bool bInIsSolverRoot)
 {
 	Name = InName;
 	ParentIndex = InParentIndex;
@@ -20,7 +19,6 @@ FBone::FBone(
 	Rotation = InOrigRotation;
 	bIsSolverRoot = bInIsSolverRoot;
 	bIsSolved = false; // default
-	ElementIndex = InElementIndex;
 }
 
 bool FBone::HasChild(const FBone* Bone)
@@ -89,7 +87,7 @@ int FRigidBody::GetNumBonesToRoot() const
 	return NumBonesToRoot; 
 }
 
-FRigidBody* FRigidBody::GetParentBody()
+FRigidBody* FRigidBody::GetParentBody() const
 {
 	if (Bone && Bone->Parent)
 	{
diff --git a/Engine/Plugins/Experimental/FullBodyIK/Source/PBIK/Private/Core/PBIKConstraint.cpp b/Engine/Plugins/Experimental/FullBodyIK/Source/PBIK/Private/Core/PBIKConstraint.cpp
--- a/Engine/Plugins/Experimental/FullBodyIK/Source/PBIK/Private/Core/PBIKConstraint.cpp
+++ b/Engine/Plugins/Experimental/FullBodyIK/Source/PBIK/Private/Core/PBIKConstraint.cpp
@@ -6,21 +6,21 @@
 
 namespace PBIK
 {
-
-	static float ALMOST_ONE = 1.0f - KINDA_SMALL_NUMBER;
 	
 FJointConstraint::FJointConstraint(FRigidBody* InA, FRigidBody* InB)
 {
 	A = InA;
 	B = InB;
 
-	FVector PinPoint = B->Bone->Position;
+	const FVector PinPoint = B->Bone->Position;
 	PinPointLocalToA = A->Rotation.Inverse() * (PinPoint - A->Position);
 	PinPointLocalToB = B->Rotation.Inverse() * (PinPoint - B->Position);
 
 	XOrig = B->RotationOrig * FVector(1, 0, 0);
 	YOrig = B->RotationOrig * FVector(0, 1, 0);
 	ZOrig = B->RotationOrig * FVector(0, 0, 1);
+
+	AngleX = AngleY = AngleZ = 0;
 }
 	
 void FJointConstraint::Solve(bool bMoveSubRoots)
@@ -29,12 +29,11 @@ void FJointConstraint::Solve(bool bMoveSubRoots)
 	FVector OffsetA;
 	FVector OffsetB;
 	FVector Correction = GetPositionCorrection(OffsetA, OffsetB);
-	//Correction = Correction * OneOverInvMassSum;
 
 	// apply rotation from correction 
 	// at pin point to both bodies
-	A->ApplyPushToRotateBody(Correction, OffsetA);//* AInvMass, OffsetA);
-	B->ApplyPushToRotateBody(-Correction, OffsetB); //* BInvMass, OffsetB);
+	A->ApplyPushToRotateBody(Correction, OffsetA);
+	B->ApplyPushToRotateBody(-Correction, OffsetB);
 
 	// enforce joint limits
 	UpdateJointLimits();
@@ -51,12 +50,12 @@ void FJointConstraint::Solve(bool bMoveSubRoots)
 	BInvMass -= !bMoveSubRoots && B->Bone->bIsSubRoot ? 1.0f : 0.0f;
 	BInvMass = BInvMass <= 0.0f ? 0.0f : BInvMass;
 
-	float InvMassSum = AInvMass + BInvMass;
+	const float InvMassSum = AInvMass + BInvMass;
 	if (FMath::IsNearlyZero(InvMassSum))
 	{
 		return; // no correction can be applied on full locked configuration
 	}
-	float OneOverInvMassSum = 1.0f / InvMassSum;
+	const float OneOverInvMassSum = 1.0f / InvMassSum;
 
 	// apply positional correction to align pin point on both Bodies
 	// NOTE: applying position correction AFTER rotation takes into
@@ -74,20 +73,20 @@ void FJointConstraint::RemoveStretch()
 	// apply position correction, keeping parent body fixed
 	FVector ToA;
 	FVector ToB;
-	FVector Correction = GetPositionCorrection(ToA, ToB);
+	const FVector Correction = GetPositionCorrection(ToA, ToB);
 	B->Position -= Correction;
 }
 
-FVector FJointConstraint::GetPositionCorrection(FVector& OutBodyToA, FVector& OutBodyToB)
+FVector FJointConstraint::GetPositionCorrection(FVector& OutBodyToA, FVector& OutBodyToB) const
 {
 	OutBodyToA = A->Rotation * PinPointLocalToA;
 	OutBodyToB = B->Rotation * PinPointLocalToB;
-	FVector PinPointOnA = A->Position + OutBodyToA;
-	FVector PinPointOnB = B->Position + OutBodyToB;
+	const FVector PinPointOnA = A->Position + OutBodyToA;
+	const FVector PinPointOnB = B->Position + OutBodyToB;
 	return (PinPointOnB - PinPointOnA);
 }
 
-void FJointConstraint::ApplyRotationCorrection(FQuat PureRotA, FQuat PureRotB)
+void FJointConstraint::ApplyRotationCorrection(FQuat PureRotA, FQuat PureRotB) const
 {
 	// https://matthias-research.github.io/pages/publications/PBDBodies.pdf/
 	// Equation 8 and 9 from "Detailed Rigid Body Simulation with Extended Position Based Dynamics"
@@ -133,35 +132,35 @@ void FJointConstraint::UpdateJointLimits()
 	J.MaxZ = J.MaxZ < J.MinZ ? J.MinZ + 1 : J.MaxZ;
 
 	// determine which axes are explicitly or implicitly locked (limited with very small allowable angle)
-	bool bLockX = (J.X == ELimitType::Locked) || (J.X == ELimitType::Limited && ((J.MaxX - J.MinX) < 2.0f));
-	bool bLockY = (J.Y == ELimitType::Locked) || (J.Y == ELimitType::Limited && ((J.MaxY - J.MinY) < 2.0f));
-	bool bLockZ = (J.Z == ELimitType::Locked) || (J.Z == ELimitType::Limited && ((J.MaxZ - J.MinZ) < 2.0f));
+	const bool bLockX = (J.X == ELimitType::Locked) || (J.X == ELimitType::Limited && ((J.MaxX - J.MinX) < 2.0f));
+	const bool bLockY = (J.Y == ELimitType::Locked) || (J.Y == ELimitType::Limited && ((J.MaxY - J.MinY) < 2.0f));
+	const bool bLockZ = (J.Z == ELimitType::Locked) || (J.Z == ELimitType::Limited && ((J.MaxZ - J.MinZ) < 2.0f));
 
 	// determine which axes should be treated as a hinge (mutually exclusive)
-	bool bXHinge = !bLockX && bLockY && bLockZ;
-	bool bYHinge = bLockX && !bLockY && bLockZ;
-	bool bZHinge = bLockX && bLockY && !bLockZ;
+	const bool bXHinge = !bLockX && bLockY && bLockZ;
+	const bool bYHinge = bLockX && !bLockY && bLockZ;
+	const bool bZHinge = bLockX && bLockY && !bLockZ;
 
 	// apply hinge corrections
 	if (bXHinge)
 	{
 		UpdateLocalRotateAxes(true, false, false);
-		FVector CrossProd = FVector::CrossProduct(XA, XB);
-		FQuat PureRotA = FQuat(CrossProd.X, CrossProd.Y, CrossProd.Z, 0.0f);
+		const FVector CrossProd = FVector::CrossProduct(XA, XB);
+		const FQuat PureRotA = FQuat(CrossProd.X, CrossProd.Y, CrossProd.Z, 0.0f);
 		ApplyRotationCorrection(PureRotA, PureRotA);
 	}
 	else if (bYHinge)
 	{
 		UpdateLocalRotateAxes(false, true, false);
-		FVector CrossProd = FVector::CrossProduct(YA, YB);
-		FQuat PureRotA = FQuat(CrossProd.X, CrossProd.Y, CrossProd.Z, 0.0f);
+		const FVector CrossProd = FVector::CrossProduct(YA, YB);
+		const FQuat PureRotA = FQuat(CrossProd.X, CrossProd.Y, CrossProd.Z, 0.0f);
 		ApplyRotationCorrection(PureRotA, PureRotA);
 	}
 	else if (bZHinge)
 	{
 		UpdateLocalRotateAxes(false, false, true);
-		FVector CrossProd = FVector::CrossProduct(ZA, ZB);
-		FQuat PureRotA = FQuat(CrossProd.X, CrossProd.Y, CrossProd.Z, 0.0f);
+		const FVector CrossProd = FVector::CrossProduct(ZA, ZB);
+		const FQuat PureRotA = FQuat(CrossProd.X, CrossProd.Y, CrossProd.Z, 0.0f);
 		ApplyRotationCorrection(PureRotA, PureRotA);
 	}
 
@@ -186,9 +185,9 @@ void FJointConstraint::UpdateJointLimits()
 	}
 
 	// enforce min/max angles
-	bool bLimitX = J.X == ELimitType::Limited && !bLockX;
-	bool bLimitY = J.Y == ELimitType::Limited && !bLockY;
-	bool bLimitZ = J.Z == ELimitType::Limited && !bLockZ;
+	const bool bLimitX = J.X == ELimitType::Limited && !bLockX;
+	const bool bLimitY = J.Y == ELimitType::Limited && !bLockY;
+	const bool bLimitZ = J.Z == ELimitType::Limited && !bLockZ;
 	if (bLimitX || bLimitY || bLimitZ)
 	{
 		DecomposeRotationAngles();
@@ -216,25 +215,25 @@ void FJointConstraint::RotateWithinLimits(
 	float CurrentAngle,
 	FVector RotAxis,
 	FVector CurVec,
-	FVector RefVec)
+	FVector RefVec) const
 {
-	bool bBeyondMin = CurrentAngle < MinAngle;
-	bool bBeyondMax = CurrentAngle > MaxAngle;
+	const bool bBeyondMin = CurrentAngle < MinAngle;
+	const bool bBeyondMax = CurrentAngle > MaxAngle;
 	if (bBeyondMin || bBeyondMax)
 	{
-		float TgtAngle = bBeyondMin ? MinAngle : MaxAngle;
-		FQuat TgtRot = FQuat(RotAxis, FMath::DegreesToRadians(TgtAngle));
-		FVector TgtVec = TgtRot * RefVec;
-		FVector TgtCross = FVector::CrossProduct(TgtVec, CurVec);
-		FQuat PureRot = FQuat(TgtCross.X, TgtCross.Y, TgtCross.Z, 0.0f);
+		const float TgtAngle = bBeyondMin ? MinAngle : MaxAngle;
+		const FQuat TgtRot = FQuat(RotAxis, FMath::DegreesToRadians(TgtAngle));
+		const FVector TgtVec = TgtRot * RefVec;
+		const FVector TgtCross = FVector::CrossProduct(TgtVec, CurVec);
+		const FQuat PureRot = FQuat(TgtCross.X, TgtCross.Y, TgtCross.Z, 0.0f);
 		ApplyRotationCorrection(PureRot, PureRot);
 	}
 }
 
 void FJointConstraint::UpdateLocalRotateAxes(bool bX, bool bY, bool bZ)
 {
-	FQuat ARot = A->Rotation * A->RotationOrig.Inverse();
-	FQuat BRot = B->Rotation * B->RotationOrig.Inverse();
+	const FQuat ARot = A->Rotation * A->RotationOrig.Inverse();
+	const FQuat BRot = B->Rotation * B->RotationOrig.Inverse();
 
 	if (bX)
 	{
@@ -257,8 +256,8 @@ void FJointConstraint::UpdateLocalRotateAxes(bool bX, bool bY, bool bZ)
 
 void FJointConstraint::DecomposeRotationAngles()
 {
-	FQuat ARot = A->Rotation * A->RotationOrig.Inverse();
-	FQuat BRot = B->Rotation * B->RotationOrig.Inverse();
+	const FQuat ARot = A->Rotation * A->RotationOrig.Inverse();
+	const FQuat BRot = B->Rotation * B->RotationOrig.Inverse();
 
 	XA = ARot * XOrig;
 	YA = ARot * YOrig;
@@ -281,10 +280,10 @@ float FJointConstraint::SignedAngleBetweenNormals(
 	const FVector& To, 
 	const FVector& Axis) const
 {
-	float FromDotTo = FVector::DotProduct(From, To);
-	float Angle = FMath::RadiansToDegrees(FMath::Acos(FromDotTo));
-	FVector Cross = FVector::CrossProduct(From, To); // TODO may be backwards?
-	float Dot = FVector::DotProduct(Cross, Axis);
+	const float FromDotTo = FVector::DotProduct(From, To);
+	const float Angle = FMath::RadiansToDegrees(FMath::Acos(FromDotTo));
+	const FVector Cross = FVector::CrossProduct(From, To); // TODO may be backwards?
+	const float Dot = FVector::DotProduct(Cross, Axis);
 	return Dot >= 0 ? Angle : -Angle;
 }
 	
@@ -324,7 +323,7 @@ void FPinConstraint::Solve(bool bMoveSubRoots)
 FVector FPinConstraint::GetPositionCorrection(FVector& OutBodyToPinPoint) const
 {
 	OutBodyToPinPoint = A->Rotation * PinPointLocalToA;
-	FVector PinPoint = A->Position + OutBodyToPinPoint;
+	const FVector PinPoint = A->Position + OutBodyToPinPoint;
 	return (GoalPoint - PinPoint) * Alpha;
 }
 
diff --git a/Engine/Plugins/Experimental/FullBodyIK/Source/PBIK/Private/Core/PBIKSolver.cpp b/Engine/Plugins/Experimental/FullBodyIK/Source/PBIK/Private/Core/PBIKSolver.cpp
--- a/Engine/Plugins/Experimental/FullBodyIK/Source/PBIK/Private/Core/PBIKSolver.cpp
+++ b/Engine/Plugins/Experimental/FullBodyIK/Source/PBIK/Private/Core/PBIKSolver.cpp
@@ -53,19 +53,19 @@ namespace PBIK
 		// if the the parent sub-root is attached to an effector, use the effector's position
 		// otherwise use the current position of the FRigidBody
 		FEffector* ParentEffector = ParentSubRoot->AttachedEffector;
-		FVector ParentSubRootPosition = ParentEffector ? ParentEffector->Position : ParentSubRoot->Position;
-		float DistToNearestSubRoot = (ParentSubRootPosition - Position).Size();
+		const FVector ParentSubRootPosition = ParentEffector ? ParentEffector->Position : ParentSubRoot->Position;
+		const float DistToNearestSubRoot = (ParentSubRootPosition - Position).Size();
 		if (DistToNearestSubRoot >= DistToSubRootOrig)
 		{
 			return; // limb is stretched
 		}
 
 		// shrink distance to reach full blend to preferred angle
-		float ScaledDistOrig = DistToSubRootOrig * 0.3f;
+		const float ScaledDistOrig = DistToSubRootOrig * 0.3f;
 		// amount squashed (clamped to scaled original length)
 		float DeltaSquash = DistToSubRootOrig - DistToNearestSubRoot;
 		DeltaSquash = DeltaSquash > ScaledDistOrig ? ScaledDistOrig : DeltaSquash;
-		float SquashPercent = DeltaSquash / ScaledDistOrig;
+		const float SquashPercent = DeltaSquash / ScaledDistOrig;
 		if (SquashPercent < 0.01f)
 		{
 			return; // limb not squashed enough
@@ -76,7 +76,6 @@ namespace PBIK
 		{
 			if (Parent->Body->J.bUsePreferredAngles)
 			{
-				FRotator Clamped = Parent->Body->J.PreferredAngles;
 				FQuat PartialRotation = FQuat::FastLerp(FQuat::Identity, FQuat(Parent->Body->J.PreferredAngles), SquashPercent);
 				Parent->Body->Rotation = Parent->Body->Rotation * PartialRotation;
 				Parent->Body->Rotation.Normalize();
@@ -150,7 +149,7 @@ void FPBIKSolver::Solve(const FPBIKSolverSettings& Settings)
 	// run constraint iterations while allowing stretch, just to get reaching pose
 	for (int32  I = 0; I < Settings.Iterations; ++I)
 	{
-		bool bMoveSubRoots = true;
+		const bool bMoveSubRoots = true;
 		for (auto Constraint : Constraints)
 		{
 			Constraint->Solve(bMoveSubRoots);
@@ -171,7 +170,7 @@ void FPBIKSolver::Solve(const FPBIKSolverSettings& Settings)
 
 		for (int32  I = 0; I < Settings.Iterations; ++I)
 		{
-			bool bMoveSubRoots = false;
+			const bool bMoveSubRoots = false;
 			for (auto Constraint : Constraints)
 			{
 				Constraint->Solve(bMoveSubRoots);
@@ -288,7 +287,7 @@ bool FPBIKSolver::InitBones()
 		}
 
 		// validate parent
-		bool bIndexInRange = Bone.ParentIndex >= 0 && Bone.ParentIndex < Bones.Num();
+		const bool bIndexInRange = Bone.ParentIndex >= 0 && Bone.ParentIndex < Bones.Num();
 		if (!ensureMsgf(bIndexInRange,
 			TEXT("PBIK: bone found with invalid parent index. Cannot initialize.")))
 		{
@@ -450,7 +449,7 @@ void FPBIKSolver::InitConstraints()
 	// this constraint is by default off in solver settings
 	if (!SolverRoot->Body->AttachedEffector) // only add if user hasn't added their own root effector
 	{
-		TSharedPtr<FPinConstraint> RootConstraint = MakeShared<FPinConstraint>(SolverRoot->Body, SolverRoot->Position);
+		const TSharedPtr<FPinConstraint> RootConstraint = MakeShared<FPinConstraint>(SolverRoot->Body, SolverRoot->Position);
 		Constraints.Add(RootConstraint);
 		RootPin = RootConstraint;
 	}
@@ -486,7 +485,7 @@ void FPBIKSolver::Reset()
 	Effectors.Empty();
 }
 
-bool FPBIKSolver::IsReadyToSimulate()
+bool FPBIKSolver::IsReadyToSimulate() const
 {
 	return bReadyToSimulate;
 }
@@ -496,10 +495,9 @@ int32 FPBIKSolver::AddBone(
 	const int32 ParentIndex,
 	const FVector& InOrigPosition,
 	const FQuat& InOrigRotation,
-	bool bIsSolverRoot,
-	const int& InElementIndex)
+	bool bIsSolverRoot)
 {
-	return Bones.Emplace(Name, ParentIndex, InOrigPosition, InOrigRotation, bIsSolverRoot, InElementIndex);
+	return Bones.Emplace(Name, ParentIndex, InOrigPosition, InOrigRotation, bIsSolverRoot);
 }
 
 int32 FPBIKSolver::AddEffector(FName BoneName)
@@ -555,6 +553,19 @@ void FPBIKSolver::GetBoneGlobalTransform(const int32 Index, FTransform& OutTrans
 	OutTransform.SetRotation(Bone.Rotation);
 }
 
+int32 FPBIKSolver::GetBoneIndex(FName BoneName) const
+{
+	for (int32 B=0; B < Bones.Num(); ++B)
+	{
+		if (Bones[B].Name == BoneName)
+		{
+			return B;
+		}
+	}
+
+	return INDEX_NONE;
+}
+
 void FPBIKSolver::SetEffectorGoal(
 	const int32 Index,
 	const FVector& InPosition,
diff --git a/Engine/Plugins/Experimental/FullBodyIK/Source/PBIK/Private/RigUnit_PBIK.cpp b/Engine/Plugins/Experimental/FullBodyIK/Source/PBIK/Private/RigUnit_PBIK.cpp
--- a/Engine/Plugins/Experimental/FullBodyIK/Source/PBIK/Private/RigUnit_PBIK.cpp
+++ b/Engine/Plugins/Experimental/FullBodyIK/Source/PBIK/Private/RigUnit_PBIK.cpp
@@ -39,18 +39,37 @@ FRigUnit_PBIK_Execute()
 
 		// reset all internal data
 		Solver.Reset();
+		
+		// create solver bone index to elements index map
+		TArray<FRigBoneElement*> BoneElements = Hierarchy->GetBones(true);
+		SolverBoneToElementIndex.Reset();
+		for (int B = 0; B < BoneElements.Num(); ++B)
+		{
+			SolverBoneToElementIndex.Add(BoneElements[B]->GetIndex());
+		}
 
 		// create bones
-		TArray<FRigBoneElement*> BoneElements = Hierarchy->GetBones(true);
 		for (int B = 0; B < BoneElements.Num(); ++B)
 		{
 			FName Name = BoneElements[B]->GetName();
-			const int ParentIndex = Hierarchy->GetFirstParent(BoneElements[B]->GetIndex());
+
+			// get the parent bone solver index
+			const int ParentElementIndex = Hierarchy->GetFirstParent(BoneElements[B]->GetIndex());
+			int ParentIndex = -1;
+			for (int P=0; P<SolverBoneToElementIndex.Num(); ++P)
+			{
+				if (SolverBoneToElementIndex[P] == ParentElementIndex)
+				{
+					ParentIndex = P;
+					break;
+				}
+			}
+			
 			const FTransform OrigTransform = Hierarchy->GetTransform(BoneElements[B], ERigTransformType::InitialGlobal);
 			const FVector InOrigPosition = OrigTransform.GetLocation();
 			const FQuat InOrigRotation = OrigTransform.GetRotation();
 			bool bIsRoot = BoneElements[B]->GetName() == Root;
-			Solver.AddBone(Name, ParentIndex, InOrigPosition, InOrigRotation, bIsRoot, BoneElements[B]->GetIndex());
+			Solver.AddBone(Name, ParentIndex, InOrigPosition, InOrigRotation, bIsRoot);
 		}
 		
 		// create effectors
@@ -73,9 +92,9 @@ FRigUnit_PBIK_Execute()
 	}
 
 	// set bones to input pose
-	for(int32 BoneIndex = 0; BoneIndex < Solver.Bones.Num(); BoneIndex++)
+	for(int32 BoneIndex = 0; BoneIndex < Solver.GetNumBones(); BoneIndex++)
 	{
-		const FTransform GlobalTransform = Hierarchy->GetGlobalTransform(Solver.Bones[BoneIndex].ElementIndex);
+		const FTransform GlobalTransform = Hierarchy->GetGlobalTransform(SolverBoneToElementIndex[BoneIndex]);
 		Solver.SetBoneTransform(BoneIndex, GlobalTransform);
 	}
 
@@ -96,14 +115,7 @@ FRigUnit_PBIK_Execute()
 
 		if(BoneSettingToSolverBoneIndex[BoneSettingIndex] == INDEX_NONE)
 		{
-			for(int32 BoneIndex = 0; BoneIndex < Solver.Bones.Num(); BoneIndex++)
-			{
-				if(Solver.Bones[BoneIndex].Name == BoneSetting.Bone)
-				{
-					BoneSettingToSolverBoneIndex[BoneSettingIndex] = BoneIndex;
-					break;
-				}
-			}
+			BoneSettingToSolverBoneIndex[BoneSettingIndex] = Solver.GetBoneIndex(BoneSetting.Bone);
 			if(BoneSettingToSolverBoneIndex[BoneSettingIndex] == INDEX_NONE)
 			{
 				continue;
@@ -136,11 +148,11 @@ FRigUnit_PBIK_Execute()
 
 	// copy transforms back
 	const bool bPropagateTransform = false;
-	for(int32 BoneIndex = 0; BoneIndex < Solver.Bones.Num(); BoneIndex++)
+	for(int32 BoneIndex = 0; BoneIndex < Solver.GetNumBones(); BoneIndex++)
 	{
 		FTransform NewTransform;
 		Solver.GetBoneGlobalTransform(BoneIndex, NewTransform);
-		Hierarchy->SetGlobalTransform(Solver.Bones[BoneIndex].ElementIndex, NewTransform, false, bPropagateTransform);
+		Hierarchy->SetGlobalTransform(SolverBoneToElementIndex[BoneIndex], NewTransform, false, bPropagateTransform);
 	}
 
 	// do all debug drawing
diff --git a/Engine/Plugins/Experimental/FullBodyIK/Source/PBIK/Public/Core/PBIKBody.h b/Engine/Plugins/Experimental/FullBodyIK/Source/PBIK/Public/Core/PBIKBody.h
--- a/Engine/Plugins/Experimental/FullBodyIK/Source/PBIK/Public/Core/PBIKBody.h
+++ b/Engine/Plugins/Experimental/FullBodyIK/Source/PBIK/Public/Core/PBIKBody.h
@@ -18,8 +18,7 @@ struct FEffector;
 struct FBone
 {
 	FName Name;
-	int ElementIndex = INDEX_NONE; // used to map to the element this bone references
-	int ParentIndex = -2; // -2 is unset, -1 for root, or 0...n otherwise
+	int32 ParentIndex = -2; // -2 is unset, -1 for root, or 0...n otherwise
 	bool bIsSolverRoot = false;
 	bool bIsSolved = false;
 	bool bIsSubRoot = false;
@@ -36,11 +35,10 @@ struct FBone
 
 	FBone(
 		const FName InName,
-		const int& InParentIndex,		// must pass -1 for root of whole skeleton
+		const int32& InParentIndex,		// must pass -1 for root of whole skeleton
 		const FVector& InOrigPosition,
 		const FQuat& InOrigRotation,
-		bool bInIsSolverRoot,
-		const int& InElementIndex = INDEX_NONE);
+		bool bInIsSolverRoot);
 
 	bool HasChild(const FBone* Bone);
 };
@@ -90,7 +88,7 @@ struct FRigidBody
 	
 private:
 
-	int NumBonesToRoot = 0;
+	int32 NumBonesToRoot = 0;
 
 public:
 
@@ -102,7 +100,7 @@ public:
 
 	int GetNumBonesToRoot() const;
 
-	FRigidBody* GetParentBody();
+	FRigidBody* GetParentBody() const;
 
 	void ApplyPushToRotateBody(const FVector& Push, const FVector& Offset);
 	
diff --git a/Engine/Plugins/Experimental/FullBodyIK/Source/PBIK/Public/Core/PBIKConstraint.h b/Engine/Plugins/Experimental/FullBodyIK/Source/PBIK/Public/Core/PBIKConstraint.h
--- a/Engine/Plugins/Experimental/FullBodyIK/Source/PBIK/Public/Core/PBIKConstraint.h
+++ b/Engine/Plugins/Experimental/FullBodyIK/Source/PBIK/Public/Core/PBIKConstraint.h
@@ -57,9 +57,9 @@ public:
 
 private:
 
-	FVector GetPositionCorrection(FVector& OutBodyToA, FVector& OutBodyToB);
+	FVector GetPositionCorrection(FVector& OutBodyToA, FVector& OutBodyToB) const;
 
-	void ApplyRotationCorrection(FQuat PureRotA, FQuat PureRotB);
+	void ApplyRotationCorrection(FQuat PureRotA, FQuat PureRotB) const;
 
 	void UpdateJointLimits();
 
@@ -69,7 +69,7 @@ private:
 		float CurrentAngle,
 		FVector RotAxis,
 		FVector CurVec,
-		FVector RefVec);
+		FVector RefVec) const;
 
 	void UpdateLocalRotateAxes(bool bX, bool bY, bool bZ);
 
diff --git a/Engine/Plugins/Experimental/FullBodyIK/Source/PBIK/Public/Core/PBIKSolver.h b/Engine/Plugins/Experimental/FullBodyIK/Source/PBIK/Public/Core/PBIKSolver.h
--- a/Engine/Plugins/Experimental/FullBodyIK/Source/PBIK/Public/Core/PBIKSolver.h
+++ b/Engine/Plugins/Experimental/FullBodyIK/Source/PBIK/Public/Core/PBIKSolver.h
@@ -87,7 +87,7 @@ public:
 
 	void Reset();
 
-	bool IsReadyToSimulate();
+	bool IsReadyToSimulate() const;
 
 	//
 	// set input / get output at runtime
@@ -106,6 +106,10 @@ public:
 
 	void GetBoneGlobalTransform(const int32 Index, FTransform& OutTransform);
 
+	int32 GetNumBones() const { return Bones.Num(); }
+
+	int32 GetBoneIndex(FName BoneName) const;
+
 	//
 	// pre-init /  setup functions
 	//
@@ -115,8 +119,7 @@ public:
 		const int32 ParentIndex,
 		const FVector& InOrigPosition,
 		const FQuat& InOrigRotation,
-		bool bIsSolverRoot,
-		const int& InElementIndex = INDEX_NONE);
+		bool bIsSolverRoot);
 
 	int32 AddEffector(FName BoneName);
 	
@@ -142,5 +145,4 @@ private:
 	
 	PBIK::FDebugDraw DebugDraw;
 	friend PBIK::FDebugDraw;
-	friend struct FRigUnit_PBIK;
 };
diff --git a/Engine/Plugins/Experimental/FullBodyIK/Source/PBIK/Public/PBIK_Shared.h b/Engine/Plugins/Experimental/FullBodyIK/Source/PBIK/Public/PBIK_Shared.h
--- a/Engine/Plugins/Experimental/FullBodyIK/Source/PBIK/Public/PBIK_Shared.h
+++ b/Engine/Plugins/Experimental/FullBodyIK/Source/PBIK/Public/PBIK_Shared.h
@@ -21,7 +21,9 @@ struct FPBIKBoneSetting
 {
 	GENERATED_BODY()
 
-	FPBIKBoneSetting() : Bone(NAME_None) {}
+	FPBIKBoneSetting()
+	:	Bone(NAME_None),
+		PreferredAngles(FRotator::ZeroRotator) {}
 
 	UPROPERTY(meta = (Constant, CustomWidget = "BoneName"))
 	FName Bone;
diff --git a/Engine/Plugins/Experimental/FullBodyIK/Source/PBIK/Public/RigUnit_PBIK.h b/Engine/Plugins/Experimental/FullBodyIK/Source/PBIK/Public/RigUnit_PBIK.h
--- a/Engine/Plugins/Experimental/FullBodyIK/Source/PBIK/Public/RigUnit_PBIK.h
+++ b/Engine/Plugins/Experimental/FullBodyIK/Source/PBIK/Public/RigUnit_PBIK.h
@@ -89,15 +89,18 @@ struct FRigUnit_PBIK : public FRigUnit_HighlevelBaseMutable
 	UPROPERTY(meta = (Input))
 	TArray<FPBIKBoneSetting> BoneSettings;
 
-	UPROPERTY(transient)
-	TArray<int32> BoneSettingToSolverBoneIndex;
-
 	UPROPERTY(meta = (Input))
 	FPBIKSolverSettings Settings;
 
 	UPROPERTY(meta = (Input))
 	FPBIKDebug Debug;
 
+	UPROPERTY(transient)
+	TArray<int32> BoneSettingToSolverBoneIndex;
+
+	UPROPERTY(transient)
+	TArray<int32> SolverBoneToElementIndex;
+
 	UPROPERTY(transient)
 	FPBIKSolver Solver;
 };
