Control Rig: UE.EditorAutomation(RunTest=Animation) Win64 - FBIKTest EndEffectorReachTest Failed: "Transform unmatches"

diff --git a/Engine/Plugins/Experimental/FullBodyIK/Source/FullBodyIK/Private/FBIKConstraintLib.cpp b/Engine/Plugins/Experimental/FullBodyIK/Source/FullBodyIK/Private/FBIKConstraintLib.cpp
--- a/Engine/Plugins/Experimental/FullBodyIK/Source/FullBodyIK/Private/FBIKConstraintLib.cpp
+++ b/Engine/Plugins/Experimental/FullBodyIK/Source/FullBodyIK/Private/FBIKConstraintLib.cpp
@@ -86,7 +86,7 @@ namespace FBIKConstraintLib
 								NewLimitConstraint.bYLimitSet = GetLimit(Constraints[Index].AngularLimit.LimitType_Y, FMath::DegreesToRadians(Constraints[Index].AngularLimit.Limit.Y), NewLimitConstraint.Limit.Y);
 								NewLimitConstraint.bZLimitSet = GetLimit(Constraints[Index].AngularLimit.LimitType_Z, FMath::DegreesToRadians(Constraints[Index].AngularLimit.Limit.Z), NewLimitConstraint.Limit.Z);
 
-								FTransform LocalTransform = Hierarchy->GetGlobalTransform(Item, true).GetRelativeTransform(Hierarchy->GetGlobalTransform(ParentItem, true));
+								FTransform LocalTransform = Hierarchy->GetInitialGlobalTransform(Item).GetRelativeTransform(Hierarchy->GetInitialGlobalTransform(ParentItem));
 								// @todo: think about 
 								// set rotation frame
 								NewLimitConstraint.BaseFrameOffset = FQuat(Constraints[Index].OffsetRotation);
diff --git a/Engine/Plugins/Experimental/FullBodyIK/Source/FullBodyIK/Private/RigUnit_FullbodyIK.cpp b/Engine/Plugins/Experimental/FullBodyIK/Source/FullBodyIK/Private/RigUnit_FullbodyIK.cpp
--- a/Engine/Plugins/Experimental/FullBodyIK/Source/FullBodyIK/Private/RigUnit_FullbodyIK.cpp
+++ b/Engine/Plugins/Experimental/FullBodyIK/Source/FullBodyIK/Private/RigUnit_FullbodyIK.cpp
@@ -533,7 +533,7 @@ FRigUnit_FullbodyIK_Execute()
 			// this means, only the last joint in the test
 			const FRigElementKey& CurrentItem = *LinkDataToHierarchyIndices.Find(LinkIndex);
 			const FTransform& LinkTransform = LinkData[LinkIndex].GetTransform();
-			Hierarchy->SetGlobalTransform(CurrentItem, LinkTransform, bPropagateToChildren);
+			Hierarchy->SetGlobalTransform(CurrentItem, LinkTransform, false, bPropagateToChildren);
 		}
 	}
 }
