Fixes odd order of Preferred Angles in UI.
Fixes uninitialized struct members in PBIKBoneSettings.

diff --git a/Engine/Plugins/Experimental/FullBodyIK/Source/PBIK/Public/PBIK_Shared.h b/Engine/Plugins/Experimental/FullBodyIK/Source/PBIK/Public/PBIK_Shared.h
--- a/Engine/Plugins/Experimental/FullBodyIK/Source/PBIK/Public/PBIK_Shared.h
+++ b/Engine/Plugins/Experimental/FullBodyIK/Source/PBIK/Public/PBIK_Shared.h
@@ -22,8 +22,11 @@ struct FPBIKBoneSetting
 	GENERATED_BODY()
 
 	FPBIKBoneSetting()
-	:	Bone(NAME_None),
-		PreferredAngles(FRotator::ZeroRotator) {}
+		: Bone(NAME_None), 
+		X(EPBIKLimitType::Free),
+		Y(EPBIKLimitType::Free),
+		Z(EPBIKLimitType::Free),
+		PreferredAngles(FVector::ZeroVector){}
 
 	UPROPERTY(meta = (Constant, CustomWidget = "BoneName"))
 	FName Bone;
@@ -57,7 +60,7 @@ struct FPBIKBoneSetting
 	UPROPERTY()
 	bool bUsePreferredAngles = false;
 	UPROPERTY()
-	FRotator PreferredAngles;
+	FVector PreferredAngles;
 
 	void CopyToCoreStruct(PBIK::FBoneSettings& Settings) const
 	{
@@ -73,6 +76,8 @@ struct FPBIKBoneSetting
 		Settings.MinZ = MinZ;
 		Settings.MaxZ = MaxZ;
 		Settings.bUsePreferredAngles = bUsePreferredAngles;
-		Settings.PreferredAngles = PreferredAngles;
+		Settings.PreferredAngles.Pitch = PreferredAngles.Y;
+		Settings.PreferredAngles.Yaw = PreferredAngles.Z;
+		Settings.PreferredAngles.Roll = PreferredAngles.X;
 	}
 };
\ No newline at end of file
