Fixes FullBodyIK crash in control rig when using anything other than bone as an effector target.
UE-104552

diff --git a/Engine/Plugins/Experimental/FullBodyIK/Source/FullBodyIK/Private/RigUnit_FullbodyIK.cpp b/Engine/Plugins/Experimental/FullBodyIK/Source/FullBodyIK/Private/RigUnit_FullbodyIK.cpp
--- a/Engine/Plugins/Experimental/FullBodyIK/Source/FullBodyIK/Private/RigUnit_FullbodyIK.cpp
+++ b/Engine/Plugins/Experimental/FullBodyIK/Source/FullBodyIK/Private/RigUnit_FullbodyIK.cpp
@@ -96,6 +96,12 @@ static void AddEffectors(URigHierarchy* Hierarchy, const FRigElementKey& Root, c
 			continue;
 		}
 
+		// can only add IK to bones
+		if (Item.Type != ERigElementType::Bone)
+		{
+			continue;
+		}
+
 		const FFBIKEndEffector& CurrentEffector = Effectors[Index];
 
 		TArray<FRigElementKey> ChainIndices;
