Deprecated jacobian fullbody IK node, renamed PBIK to "Full Body IK"

diff --git a/Engine/Plugins/Experimental/FullBodyIK/Source/FullBodyIK/Private/RigUnit_FullbodyIK.h b/Engine/Plugins/Experimental/FullBodyIK/Source/FullBodyIK/Private/RigUnit_FullbodyIK.h
--- a/Engine/Plugins/Experimental/FullBodyIK/Source/FullBodyIK/Private/RigUnit_FullbodyIK.h
+++ b/Engine/Plugins/Experimental/FullBodyIK/Source/FullBodyIK/Private/RigUnit_FullbodyIK.h
@@ -99,7 +99,7 @@ struct FRigUnit_FullbodyIK_WorkData
 /**
  * Based on Jacobian solver at core, this can solve multi chains within a root using multi effectors
  */
-USTRUCT(meta=(DisplayName="Fullbody IK", Category="Hierarchy", Keywords="Multi, Effector, N-Chain, FB, IK"))
+USTRUCT(meta=(DisplayName="Fullbody IK", Category="Hierarchy", Keywords="Multi, Effector, N-Chain, FB, IK", Deprecated = "5.0"))
 struct FRigUnit_FullbodyIK : public FRigUnit_HighlevelBaseMutable
 {
 	GENERATED_BODY()
diff --git a/Engine/Plugins/Experimental/FullBodyIK/Source/PBIK/Public/RigUnit_PBIK.h b/Engine/Plugins/Experimental/FullBodyIK/Source/PBIK/Public/RigUnit_PBIK.h
--- a/Engine/Plugins/Experimental/FullBodyIK/Source/PBIK/Public/RigUnit_PBIK.h
+++ b/Engine/Plugins/Experimental/FullBodyIK/Source/PBIK/Public/RigUnit_PBIK.h
@@ -65,7 +65,7 @@ struct FPBIKEffector
 	float StrengthAlpha = 1.0f;
 };
 
-USTRUCT(meta=(DisplayName="Position Based IK", Category="Hierarchy", Keywords="Position Based, PBIK, IK, Full Body"))
+USTRUCT(meta=(DisplayName="Full Body IK", Category="Hierarchy", Keywords="Position Based, PBIK, IK, Full Body, Multi, Effector, N-Chain, FB"))
 struct FRigUnit_PBIK : public FRigUnit_HighlevelBaseMutable
 {
 	GENERATED_BODY()
