Adding PBIK to IKRig

diff --git a/Engine/Plugins/Animation/IKRig/Source/IKRig/IKRig.Build.cs b/Engine/Plugins/Animation/IKRig/Source/IKRig/IKRig.Build.cs
--- a/Engine/Plugins/Animation/IKRig/Source/IKRig/IKRig.Build.cs
+++ b/Engine/Plugins/Animation/IKRig/Source/IKRig/IKRig.Build.cs
@@ -18,7 +18,7 @@ public class IKRig : ModuleRules
 				"AnimationCore", 
 				"AnimGraphRuntime",
 				
-				"FullBodyIK"
+				"PBIK"
 			}
 			);
 
@@ -27,6 +27,8 @@ public class IKRig : ModuleRules
 			{
 				"ControlRig",
 				"Core",
+				
+				"PBIK"
 			}
 			);
 
diff --git a/Engine/Plugins/Animation/IKRig/Source/IKRig/Private/AnimNodes/AnimNode_IKRetargeter.cpp b/Engine/Plugins/Animation/IKRig/Source/IKRig/Private/AnimNodes/AnimNode_IKRetargeter.cpp
--- a/Engine/Plugins/Animation/IKRig/Source/IKRig/Private/AnimNodes/AnimNode_IKRetargeter.cpp
+++ b/Engine/Plugins/Animation/IKRig/Source/IKRig/Private/AnimNodes/AnimNode_IKRetargeter.cpp
@@ -33,6 +33,7 @@ void FAnimNode_IKRetargeter::Evaluate_AnyThread(FPoseContext& Output)
 
 	if (!bIsInitialized)
 	{
+		Output.ResetToRefPose();
 		return;
 	}
 
@@ -138,7 +139,7 @@ void FAnimNode_IKRetargeter::InitializeRetargetData(const UAnimInstance* InAnimI
 		const_cast<UAnimInstance*>(InAnimInstance));
 	
 	// made it!
-	bIsInitialized = true;
+	bIsInitialized = CurrentlyUsedRetargeter->bIsLoadedAndValid;
 }
 
 void FAnimNode_IKRetargeter::CopyBoneTransformsFromSource(USkeletalMeshComponent* TargetMeshComponent)
@@ -189,3 +190,8 @@ void FAnimNode_IKRetargeter::CopyBoneTransformsFromSource(USkeletalMeshComponent
 	// ref skeleton is need for parent index lookups later, so store it now
 	CurrentlyUsedSourceMesh = SourceMeshComp->SkeletalMesh;
 }
+
+USkeletalMeshComponent* FAnimNode_IKRetargeter::GetSourceMesh() const
+{
+	return nullptr;
+}
diff --git a/Engine/Plugins/Animation/IKRig/Source/IKRig/Private/AnimNodes/AnimNode_IKRig.cpp b/Engine/Plugins/Animation/IKRig/Source/IKRig/Private/AnimNodes/AnimNode_IKRig.cpp
--- a/Engine/Plugins/Animation/IKRig/Source/IKRig/Private/AnimNodes/AnimNode_IKRig.cpp
+++ b/Engine/Plugins/Animation/IKRig/Source/IKRig/Private/AnimNodes/AnimNode_IKRig.cpp
@@ -162,15 +162,10 @@ void FAnimNode_IKRig::PreUpdate(const UAnimInstance* InAnimInstance)
 {
 	// pull all the goals out of any goal creators on the owning actor
 	// this is done on the main thread because we're talking to actor components here
-	GoalsFromGoalCreators.Empty();
+	GoalsFromGoalCreators.Reset();
 	for (IIKGoalCreatorInterface* GoalCreator : GoalCreators)
 	{
-		TMap<FName, FIKRigGoal> OutGoals;
-		GoalCreator->GetIKGoals_Implementation(OutGoals);
-		for (const TPair<FName, FIKRigGoal>& GoalPair : OutGoals)
-		{
-			GoalsFromGoalCreators.Add(GoalPair);
-		}
+		GoalCreator->AddIKGoals_Implementation(GoalsFromGoalCreators);
 	}
 }
 
@@ -181,8 +176,7 @@ bool FAnimNode_IKRig::RebuildGoalList()
 		return false;
 	}
 	
-	TArray<FIKRigEffectorGoal> AllGoalNamesInIKRig;
-	RigDefinitionAsset->GetGoalNamesFromSolvers(AllGoalNamesInIKRig);
+	const TArray<FIKRigEffectorGoal>& AllGoalNamesInIKRig = RigDefinitionAsset->GetEffectorGoals();
 	const int32 NumGoalsInRig = AllGoalNamesInIKRig.Num();
 	if (Goals.Num() != NumGoalsInRig)
 	{
@@ -201,15 +195,10 @@ FName FAnimNode_IKRig::GetGoalName(int32 Index) const
 {
 	if (RigDefinitionAsset)
 	{
-		TArray<FIKRigEffectorGoal> Names;
-		RigDefinitionAsset->GetGoalNamesFromSolvers(Names);
-		if (Names.IsValidIndex(Index))
-		{
-			return Names[Index].Goal;
-		}
+		return RigDefinitionAsset->GetGoalName(Index);
 	}
 
-	return FName(TEXT("Invalid"));
+	return NAME_None;
 }
 
 void FAnimNode_IKRig::Update_AnyThread(const FAnimationUpdateContext& Context)
diff --git a/Engine/Plugins/Animation/IKRig/Source/IKRig/Private/IKRigDefinition.cpp b/Engine/Plugins/Animation/IKRig/Source/IKRig/Private/IKRigDefinition.cpp
--- a/Engine/Plugins/Animation/IKRig/Source/IKRig/Private/IKRigDefinition.cpp
+++ b/Engine/Plugins/Animation/IKRig/Source/IKRig/Private/IKRigDefinition.cpp
@@ -17,32 +17,69 @@ FBoneChain* FRetargetDefinition::GetBoneChainByName(FName ChainName)
 	return nullptr;
 }
 
-void UIKRigDefinition::GetGoalNamesFromSolvers(TArray<FIKRigEffectorGoal>& OutGoalNames) const
+TArray<FIKRigEffectorGoal>& UIKRigDefinition::GetEffectorGoals()
 {
-	TSet<FIKRigEffectorGoal> GoalNames;
-	for (UIKRigSolver* Solver : Solvers)
+	UpdateGoalNameArray();
+	return EffectorGoals;
+}
+
+#if WITH_EDITOR
+void UIKRigDefinition::PostEditChangeProperty(FPropertyChangedEvent& PropertyChangedEvent)
+{
+	// update list of goal names whenever a solver is modified
+	const FName PropertyName = (PropertyChangedEvent.Property ? PropertyChangedEvent.Property->GetFName() : NAME_None);
+	if ((PropertyName == GET_MEMBER_NAME_STRING_CHECKED(UIKRigDefinition, Solvers)))
 	{
-		if (Solver)
+		bEffectorGoalsDirty = true;
+	}
+
+	Super::PostEditChangeProperty(PropertyChangedEvent);
+}
+#endif
+
+FName UIKRigDefinition::GetBoneNameForGoal(const FName& GoalName)
+{
+	UpdateGoalNameArray();
+	
+	for (const FIKRigEffectorGoal& EffectorGoal : EffectorGoals)
+	{
+		if (EffectorGoal.Goal == GoalName)
 		{
-			Solver->CollectGoalNames(GoalNames);
+			return EffectorGoal.Bone;
 		}
 	}
+	
+	return NAME_None;
+}
 
-	// user code needs to use indices, so we bake set into an array
-	OutGoalNames = GoalNames.Array();
+FName UIKRigDefinition::GetGoalName(int32 GoalIndex)
+{
+	UpdateGoalNameArray();
+	
+	if (!EffectorGoals.IsValidIndex(GoalIndex))
+	{
+		return NAME_None;
+	}
+
+	return EffectorGoals[GoalIndex].Goal;
 }
 
-FName UIKRigDefinition::GetBoneNameForGoal(FName GoalName) const
+void UIKRigDefinition::UpdateGoalNameArray()
 {
-	TArray<FIKRigEffectorGoal> AllGoalNames;
-	GetGoalNamesFromSolvers(AllGoalNames);
-	for (const FIKRigEffectorGoal& Names : AllGoalNames)
+	if (!bEffectorGoalsDirty)
 	{
-		if (Names.Goal == GoalName)
+		return;	
+	}
+	
+	EffectorGoals.Reset();
+	for (UIKRigSolver* Solver : Solvers)
+	{
+		if (Solver)
 		{
-			return Names.Bone;
+			Solver->AddGoalsInSolver(EffectorGoals);
 		}
 	}
 
-	return NAME_None;
+	bEffectorGoalsDirty = false;
 }
+
diff --git a/Engine/Plugins/Animation/IKRig/Source/IKRig/Private/IKRigProcessor.cpp b/Engine/Plugins/Animation/IKRig/Source/IKRig/Private/IKRigProcessor.cpp
--- a/Engine/Plugins/Animation/IKRig/Source/IKRig/Private/IKRigProcessor.cpp
+++ b/Engine/Plugins/Animation/IKRig/Source/IKRig/Private/IKRigProcessor.cpp
@@ -23,36 +23,35 @@ void UIKRigProcessor::Initialize(UIKRigDefinition* InRigDefinition)
 	Skeleton = InRigDefinition->Skeleton; // trivial copy assignment operator for POD
 
 	// initialize goal names based on solvers
-	TArray<FIKRigEffectorGoal> GoalNames;
-	InRigDefinition->GetGoalNamesFromSolvers(GoalNames);
-	GoalContainer.InitializeGoalsFromNames(GoalNames);
-	for (const FIKRigEffectorGoal& GoalName : GoalNames)
+	TArray<FIKRigEffectorGoal>& EffectorGoals = InRigDefinition->GetEffectorGoals();
+	GoalContainer.InitializeGoalsFromNames(EffectorGoals);
+	for (const FIKRigEffectorGoal& EffectorGoal : EffectorGoals)
 	{
 		FGoalBone NewGoalBone;
-		NewGoalBone.BoneName = GoalName.Goal;
-		NewGoalBone.BoneIndex = Skeleton.GetBoneIndexFromName(GoalName.Bone);
+		NewGoalBone.BoneName = EffectorGoal.Bone;
+		NewGoalBone.BoneIndex = Skeleton.GetBoneIndexFromName(EffectorGoal.Bone);
 
 		// validate that the skeleton we are trying to solve this goal on contains the bone the goal expects
 		if (NewGoalBone.BoneIndex == INDEX_NONE)
 		{
 			UE_LOG(LogTemp, Warning, TEXT("IK Rig, %s has a Goal, '%s' that references an unknown bone, '%s'. Cannot evaluate."),
-				*GetName(), *GoalName.Goal.ToString(), *GoalName.Bone.ToString());
+				*GetName(), *EffectorGoal.Goal.ToString(), *EffectorGoal.Bone.ToString());
 			return;
 		}
 
 		// validate that there is not already a different goal, with the same name, that is using a different bone
 		// (all goals with the same name must reference the same bone within a single IK Rig)
-		if (const FGoalBone* Bone = GoalBones.Find(GoalName.Goal))
+		if (const FGoalBone* Bone = GoalBones.Find(EffectorGoal.Goal))
 		{
 			if (Bone->BoneName != NewGoalBone.BoneName)
 			{
 				UE_LOG(LogTemp, Warning, TEXT("IK Rig, %s has a Goal, '%s' that references different bones in different solvers, '%s' and '%s'. Cannot evaluate."),
-                *GetName(), *GoalName.Goal.ToString(), *Bone->BoneName.ToString(), *NewGoalBone.BoneName.ToString());
+                *GetName(), *EffectorGoal.Goal.ToString(), *Bone->BoneName.ToString(), *NewGoalBone.BoneName.ToString());
 				return;
 			}
 		}
 		
-		GoalBones.Add(GoalName.Goal, NewGoalBone);
+		GoalBones.Add(EffectorGoal.Goal, NewGoalBone);
 	}
 
 	// create copies of all the solvers in the IK rig
@@ -99,29 +98,18 @@ void UIKRigProcessor::SetIKGoal(const FIKRigGoal& InGoal)
 void UIKRigProcessor::Solve()
 {
 	check(bInitialized);
-	
-	DrawInterface.Reset();
 
 	// blend goals towards input pose by alpha
-	for (TPair<FName, FIKRigGoal>& GoalPair : GoalContainer.Goals)
-	{
-		FIKRigGoal& Goal = GoalPair.Value;
-		const FGoalBone& GoalBone = GoalBones[Goal.Name];
-		const FTransform& InputPoseBoneTransform = Skeleton.CurrentPoseGlobal[GoalBone.BoneIndex];
-		Goal.Position = FMath::Lerp(
-			InputPoseBoneTransform.GetTranslation(),
-			Goal.Position,
-			Goal.PositionAlpha);
-		Goal.Rotation = FQuat::FastLerp(
-			InputPoseBoneTransform.GetRotation(),
-			Goal.Rotation.Quaternion(),
-			Goal.RotationAlpha).Rotator();
-	}
+	BlendGoalsByAlpha();
 
 	// run all the solvers
+	DrawInterface.Reset();
 	for (UIKRigSolver* Solver : Solvers)
 	{
-		Solver->Solve(Skeleton, GoalContainer, &DrawInterface);
+		if (Solver->bEnabled)
+		{
+			Solver->Solve(Skeleton, GoalContainer, &DrawInterface);
+		}
 	}
 
 	// make sure rotations are normalized coming out
@@ -155,3 +143,30 @@ bool UIKRigProcessor::GetBoneForGoal(FName GoalName, FGoalBone& OutBone) const
 	}
 	return false;
 }
+
+void UIKRigProcessor::BlendGoalsByAlpha()
+{
+	for (TPair<FName, FIKRigGoal>& GoalPair : GoalContainer.Goals)
+	{
+		if (!GoalBones.Contains(GoalPair.Key))
+		{
+			// user is changing goals after initialization
+			// not necessarily a bad thing, but new goal names won't work until re-init
+			continue;
+		}
+		
+		FIKRigGoal& Goal = GoalPair.Value;
+		const FGoalBone& GoalBone = GoalBones[Goal.Name];
+		const FTransform& InputPoseBoneTransform = Skeleton.CurrentPoseGlobal[GoalBone.BoneIndex];
+		
+		Goal.FinalBlendedPosition = FMath::Lerp(
+            InputPoseBoneTransform.GetTranslation(),
+            Goal.Position,
+            Goal.PositionAlpha);
+		
+		Goal.FinalBlendedRotation = FQuat::FastLerp(
+            InputPoseBoneTransform.GetRotation(),
+            Goal.Rotation.Quaternion(),
+            Goal.RotationAlpha);
+	}
+}
diff --git a/Engine/Plugins/Animation/IKRig/Source/IKRig/Private/IKRigSolver.cpp b/Engine/Plugins/Animation/IKRig/Source/IKRig/Private/IKRigSolver.cpp
--- a/Engine/Plugins/Animation/IKRig/Source/IKRig/Private/IKRigSolver.cpp
+++ b/Engine/Plugins/Animation/IKRig/Source/IKRig/Private/IKRigSolver.cpp
@@ -1,3 +1,18 @@
 // Copyright Epic Games, Inc. All Rights Reserved.
 
-#include "IKRigSolver.h"
\ No newline at end of file
+#include "IKRigSolver.h"
+
+
+void UIKRigSolver::AddGoalToArrayNoDuplicates(const FIKRigEffectorGoal& GoalToAdd, TArray<FIKRigEffectorGoal>& OutGoals)
+{
+	// need to treat OutGoals like an TSet (with no duplicates) but ordered so users can index into it
+	for (const FIKRigEffectorGoal& EffectorGoal : OutGoals)
+	{
+		if (EffectorGoal.Goal == GoalToAdd.Goal)
+		{
+			return;
+		}
+	}
+	
+	OutGoals.Add(GoalToAdd);
+}
diff --git a/Engine/Plugins/Animation/IKRig/Source/IKRig/Private/Solvers/ConstraintSolver.cpp b/Engine/Plugins/Animation/IKRig/Source/IKRig/Private/Solvers/ConstraintSolver.cpp
deleted file mode 100644
diff --git a/Engine/Plugins/Animation/IKRig/Source/IKRig/Private/Solvers/IKRigFBIKSolver.cpp b/Engine/Plugins/Animation/IKRig/Source/IKRig/Private/Solvers/IKRigFBIKSolver.cpp
deleted file mode 100644
diff --git a/Engine/Plugins/Animation/IKRig/Source/IKRig/Private/Solvers/IKRig_PBIKSolver.cpp b/Engine/Plugins/Animation/IKRig/Source/IKRig/Private/Solvers/IKRig_PBIKSolver.cpp
new file mode 100644
--- /dev/null
+++ b/Engine/Plugins/Animation/IKRig/Source/IKRig/Private/Solvers/IKRig_PBIKSolver.cpp
@@ -0,0 +1,126 @@
+// Copyright Epic Games, Inc. All Rights Reserved.
+
+#include "Solvers/IKRig_PBIKSolver.h"
+
+#include "IKRigDataTypes.h"
+#include "IKRigSkeleton.h"
+
+void UIKRigPBIKSolver::Initialize(const FIKRigSkeleton& InSkeleton)
+{	
+	// check how many effectors are assigned to a bone
+	int NumEffectors = 0;
+	for (const FIKRig_PBIKEffector& Effector : Effectors)
+	{
+		if (InSkeleton.GetBoneIndexFromName(Effector.Goal.Bone) != INDEX_NONE)
+		{
+			++NumEffectors; // bone is set and exists!
+		}
+	}
+
+	// validate inputs are ready to be initialized
+	const bool bHasEffectors = NumEffectors > 0;
+	const bool bRootIsAssigned = RootBone != NAME_None;
+	if (!(bHasEffectors && bRootIsAssigned))
+	{
+		return; // not setup yet
+	}
+
+	// reset all internal data
+	Solver.Reset();
+
+	// create bones
+	for (int BoneIndex = 0; BoneIndex < InSkeleton.BoneNames.Num(); ++BoneIndex)
+	{
+		const FName& Name = InSkeleton.BoneNames[BoneIndex];
+
+		// get the parent bone solver index
+		const int32 ParentIndex = InSkeleton.GetParentIndex(BoneIndex);
+		const FTransform OrigTransform = InSkeleton.RefPoseGlobal[BoneIndex];
+		const FVector InOrigPosition = OrigTransform.GetLocation();
+		const FQuat InOrigRotation = OrigTransform.GetRotation();
+		const bool bIsRoot = Name == RootBone;
+		Solver.AddBone(Name, ParentIndex, InOrigPosition, InOrigRotation, bIsRoot);
+	}
+
+	// create effectors
+	for (FIKRig_PBIKEffector& Effector : Effectors)
+	{
+		Effector.IndexInSolver = Solver.AddEffector(Effector.Goal.Bone);
+	}
+		
+	// initialize solver
+	Solver.Initialize();
+}
+
+void UIKRigPBIKSolver::Solve(
+    FIKRigSkeleton& IKRigSkeleton,
+    const FIKRigGoalContainer& Goals,
+    FControlRigDrawInterface* InOutDrawInterface)
+{
+	if (!Solver.IsReadyToSimulate())
+	{
+		return;
+	}
+
+	if (Solver.GetNumBones() != IKRigSkeleton.BoneNames.Num())
+	{
+		return;
+	}
+
+	// set bones to input pose
+	for(int32 BoneIndex = 0; BoneIndex < Solver.GetNumBones(); BoneIndex++)
+	{
+		const FTransform GlobalTransform = IKRigSkeleton.CurrentPoseGlobal[BoneIndex];
+		Solver.SetBoneTransform(BoneIndex, GlobalTransform);
+	}
+
+	// update bone settings
+	for (const FPBIKBoneSetting& BoneSetting : BoneSettings)
+	{
+		const int32 BoneIndex = Solver.GetBoneIndex(BoneSetting.Bone);
+		if (PBIK::FBoneSettings* InternalSettings = Solver.GetBoneSettings(BoneIndex))
+		{
+			BoneSetting.CopyToCoreStruct(*InternalSettings);
+		}
+	}
+
+	// update effectors
+	const float OffsetAlpha = 1.0f; // this is constant because IKRig manages offset alphas itself
+	for (const FIKRig_PBIKEffector& Effector : Effectors)
+	{
+		if (Effector.IndexInSolver < 0)
+		{
+			continue;
+		}
+		FIKRigGoal Goal;
+		Goals.GetGoalByName(Effector.Goal.Goal, Goal);
+		Solver.SetEffectorGoal(
+			Effector.IndexInSolver,
+			Goal.FinalBlendedPosition,
+			Goal.FinalBlendedRotation,
+			OffsetAlpha,
+			Effector.StrengthAlpha);
+	}
+
+	// solve
+	Solver.Solve(Settings);
+
+	// copy transforms back
+	for(int32 BoneIndex = 0; BoneIndex < Solver.GetNumBones(); BoneIndex++)
+	{
+		FTransform OutTransform;
+		Solver.GetBoneGlobalTransform(BoneIndex, OutTransform);
+		IKRigSkeleton.CurrentPoseGlobal[BoneIndex] = OutTransform;
+	}
+	
+	// do all debug drawing
+	//Debug.Draw(Context.DrawInterface, &Solver);
+}
+
+void UIKRigPBIKSolver::AddGoalsInSolver(TArray<FIKRigEffectorGoal>& OutGoals) const
+{
+	for (const FIKRig_PBIKEffector& Effector : Effectors)
+	{
+		AddGoalToArrayNoDuplicates(Effector.Goal,OutGoals);
+	}
+}
\ No newline at end of file
diff --git a/Engine/Plugins/Animation/IKRig/Source/IKRig/Private/Solvers/IKRig_SetTransform.cpp b/Engine/Plugins/Animation/IKRig/Source/IKRig/Private/Solvers/IKRig_SetTransform.cpp
new file mode 100644
--- /dev/null
+++ b/Engine/Plugins/Animation/IKRig/Source/IKRig/Private/Solvers/IKRig_SetTransform.cpp
@@ -0,0 +1,42 @@
+// Copyright Epic Games, Inc. All Rights Reserved.
+
+#include "Solvers/IKRig_SetTransform.h"
+#include "IKRigDataTypes.h"
+#include "IKRigSkeleton.h"
+
+
+void UIKRig_SetTransform::Initialize(const FIKRigSkeleton& IKRigSkeleton)
+{
+	BoneIndex = IKRigSkeleton.GetBoneIndexFromName(Effector.Bone);
+}
+
+void UIKRig_SetTransform::Solve(
+	FIKRigSkeleton& IKRigSkeleton,
+	const FIKRigGoalContainer& Goals,
+	FControlRigDrawInterface* InOutDrawInterface)
+{
+	FIKRigGoal OutGoal;
+	if (!Goals.GetGoalByName(Effector.Goal, OutGoal))
+	{
+		return;
+	}
+
+	FTransform& CurrentTransform = IKRigSkeleton.CurrentPoseGlobal[BoneIndex];
+
+	if (bEnablePosition)
+	{
+		CurrentTransform.SetLocation(OutGoal.FinalBlendedPosition);
+	}
+	
+	if (bEnableRotation)
+	{
+		CurrentTransform.SetRotation(OutGoal.FinalBlendedRotation);
+	}
+
+	IKRigSkeleton.PropagateGlobalPoseBelowBone(BoneIndex);
+}
+
+void UIKRig_SetTransform::AddGoalsInSolver(TArray<FIKRigEffectorGoal>& OutGoals) const
+{
+	AddGoalToArrayNoDuplicates(Effector,OutGoals);
+}
diff --git a/Engine/Plugins/Animation/IKRig/Source/IKRig/Private/Solvers/TransformSolver.cpp b/Engine/Plugins/Animation/IKRig/Source/IKRig/Private/Solvers/TransformSolver.cpp
deleted file mode 100644
diff --git a/Engine/Plugins/Animation/IKRig/Source/IKRig/Public/ActorComponents/IKRigComponent.h b/Engine/Plugins/Animation/IKRig/Source/IKRig/Public/ActorComponents/IKRigComponent.h
--- a/Engine/Plugins/Animation/IKRig/Source/IKRig/Public/ActorComponents/IKRigComponent.h
+++ b/Engine/Plugins/Animation/IKRig/Source/IKRig/Public/ActorComponents/IKRigComponent.h
@@ -55,7 +55,7 @@ public:
 
 	// BEGIN IIKRigGoalCreator interface
 	/** Called from the IK Rig Anim node to obtain all the goals that have been set on this actor.*/
-	void GetIKGoals_Implementation(TMap<FName, FIKRigGoal>& OutGoals) { OutGoals = GoalContainer.Goals; };
+	void AddIKGoals_Implementation(TMap<FName, FIKRigGoal>& OutGoals) override { OutGoals.Append(GoalContainer.Goals); };
 	// END IIKRigGoalCreator interface
 
 private:
diff --git a/Engine/Plugins/Animation/IKRig/Source/IKRig/Public/ActorComponents/IKRigInterface.h b/Engine/Plugins/Animation/IKRig/Source/IKRig/Public/ActorComponents/IKRigInterface.h
--- a/Engine/Plugins/Animation/IKRig/Source/IKRig/Public/ActorComponents/IKRigInterface.h
+++ b/Engine/Plugins/Animation/IKRig/Source/IKRig/Public/ActorComponents/IKRigInterface.h
@@ -19,7 +19,7 @@ class IIKGoalCreatorInterface
 
 public:
 	
-	/** Provide a Goal container for callers to read from.*/
+	/** Add your own goals to the OutGoals map (careful not to remove existing goals in the map!)*/
 	UFUNCTION(BlueprintCallable, BlueprintNativeEvent, Category=IKRigGoals)
-    void GetIKGoals(TMap<FName, FIKRigGoal>& OutGoals);
+    void AddIKGoals(TMap<FName, FIKRigGoal>& OutGoals);
 };
\ No newline at end of file
diff --git a/Engine/Plugins/Animation/IKRig/Source/IKRig/Public/AnimNodes/AnimNode_IKRetargeter.h b/Engine/Plugins/Animation/IKRig/Source/IKRig/Public/AnimNodes/AnimNode_IKRetargeter.h
--- a/Engine/Plugins/Animation/IKRig/Source/IKRig/Public/AnimNodes/AnimNode_IKRetargeter.h
+++ b/Engine/Plugins/Animation/IKRig/Source/IKRig/Public/AnimNodes/AnimNode_IKRetargeter.h
@@ -15,6 +15,10 @@ struct IKRIG_API FAnimNode_IKRetargeter : public FAnimNode_Base
 	/** The Skeletal Mesh Component to retarget animation from. Assumed to be animated and tick BEFORE this anim instance.*/
 	UPROPERTY(BlueprintReadWrite, transient, Category=Settings, meta=(PinShownByDefault))
 	TWeakObjectPtr<USkeletalMeshComponent> SourceMeshComponent = nullptr;
+
+	/** Map of chain names to per-chain retarget settings (can be modified at runtime).*/
+	UPROPERTY(BlueprintReadWrite, transient, Category=Settings, meta=(PinShownByDefault))
+	TMap<FName, FIKRetargetChainSettings> ChainSettings;
 	
 	/** Retarget asset to use. Must define a Source and Target IK Rig compatible with the SourceMeshComponent and current anim instance.*/
 	UPROPERTY(EditAnywhere, Category = Settings)
@@ -24,6 +28,14 @@ struct IKRIG_API FAnimNode_IKRetargeter : public FAnimNode_Base
 	UPROPERTY(EditAnywhere, Category = Settings)
 	bool bEnableIK = true;
 
+	/** Retarget asset to use. Must define a Source and Target IK Rig compatible with the SourceMeshComponent and current anim instance.*/
+	UPROPERTY(EditAnywhere, Category = Settings)
+	bool bAutoFindSourceMeshByTag = false;
+
+	/** Retarget asset to use. Must define a Source and Target IK Rig compatible with the SourceMeshComponent and current anim instance.*/
+	UPROPERTY(EditAnywhere, Category = Settings, meta = (EditCondition="bAutoFindSourceMeshByTag"))
+	FName SourceMeshComponentTag = "RetargetSource";
+	
 	// FAnimNode_Base interface
 	virtual void Initialize_AnyThread(const FAnimationInitializeContext& Context) override;
 	virtual void CacheBones_AnyThread(const FAnimationCacheBonesContext& Context) override;
@@ -35,9 +47,11 @@ struct IKRIG_API FAnimNode_IKRetargeter : public FAnimNode_Base
 	// End of FAnimNode_Base interface
 
 private:
+	
 	void EnsureInitialized(const UAnimInstance* InAnimInstance);
 	void InitializeRetargetData(const UAnimInstance* InAnimInstance);
 	void CopyBoneTransformsFromSource(USkeletalMeshComponent* TargetMeshComponent);
+	USkeletalMeshComponent* GetSourceMesh() const;
 
 	// indicates that all prerequisites are met and node is ready to operate
 	UPROPERTY(Transient)
@@ -53,5 +67,4 @@ private:
 
 	// cached transforms, copied on the game thread
 	TArray<FTransform> SourceMeshComponentSpaceBoneTransforms;
-	TArray<FTransform> TargetMeshComponentSpaceBoneTransforms;
 };
diff --git a/Engine/Plugins/Animation/IKRig/Source/IKRig/Public/IKRigDataTypes.h b/Engine/Plugins/Animation/IKRig/Source/IKRig/Public/IKRigDataTypes.h
--- a/Engine/Plugins/Animation/IKRig/Source/IKRig/Public/IKRigDataTypes.h
+++ b/Engine/Plugins/Animation/IKRig/Source/IKRig/Public/IKRigDataTypes.h
@@ -13,6 +13,25 @@
 struct FIKRigEffectorGoal;
 struct FIKRigHierarchy;
 
+USTRUCT()
+struct IKRIG_API FIKRigEffectorGoal
+{
+	GENERATED_BODY()
+
+	FIKRigEffectorGoal()
+        : Goal("DefaultGoal"),
+          Bone(NAME_None)
+	{
+	}
+
+	UPROPERTY(EditAnywhere, Category = FIKRigEffector)
+	FName Goal;
+
+	UPROPERTY(EditAnywhere, Category = FIKRigEffector)
+	FName Bone;
+
+	bool operator==(const FIKRigEffectorGoal& Other) const { return Goal == Other.Goal; }
+};
 
 USTRUCT(Blueprintable)
 struct IKRIG_API FIKRigGoal
@@ -39,19 +58,29 @@ struct IKRIG_API FIKRigGoal
 	UPROPERTY(EditAnywhere, BlueprintReadWrite, Category = FIKRigGoal, meta = (ClampMin = "0.0", ClampMax = "1.0", UIMin = "0.0", UIMax = "1.0"))
 	float RotationAlpha;
 
+	UPROPERTY(Transient)
+	FVector FinalBlendedPosition;
+
+	UPROPERTY(Transient)
+	FQuat FinalBlendedRotation;
+
 	FIKRigGoal()
     : Name(NAME_None),
 	Position(ForceInitToZero),
     Rotation(FRotator::ZeroRotator),
     PositionAlpha(1.f),
-    RotationAlpha(0.0f){}
+    RotationAlpha(0.0f),
+	FinalBlendedPosition(ForceInitToZero),
+	FinalBlendedRotation(FQuat::Identity){}
 	
 	FIKRigGoal(const FName& GoalName)
     : Name(GoalName),
 	Position(ForceInitToZero),
     Rotation(FRotator::ZeroRotator),
     PositionAlpha(1.f),
-    RotationAlpha(0.0f){}
+    RotationAlpha(0.0f),
+	FinalBlendedPosition(ForceInitToZero),
+    FinalBlendedRotation(FQuat::Identity){}
 
 	FIKRigGoal(
         const FName& Name,
@@ -63,12 +92,14 @@ struct IKRIG_API FIKRigGoal
           Position(Position),
           Rotation(Rotation.Rotator()),
           PositionAlpha(PositionAlpha),
-          RotationAlpha(RotationAlpha){}
+          RotationAlpha(RotationAlpha),
+		  FinalBlendedPosition(Position),
+		  FinalBlendedRotation(Rotation){}
 
 	FString ToString() const
 	{
 		return FString::Printf(TEXT("Name=%s, Pos=(%s, Alpha=%3.3f), Rot=(%s, Alpha=%3.3f)"),
-			*Name.ToString(), *Position.ToString(), PositionAlpha, *Rotation.ToString(), RotationAlpha);
+			*Name.ToString(), *FinalBlendedPosition.ToString(), PositionAlpha, *FinalBlendedRotation.ToString(), RotationAlpha);
 	}
 };
 
diff --git a/Engine/Plugins/Animation/IKRig/Source/IKRig/Public/IKRigDefinition.h b/Engine/Plugins/Animation/IKRig/Source/IKRig/Public/IKRigDefinition.h
--- a/Engine/Plugins/Animation/IKRig/Source/IKRig/Public/IKRigDefinition.h
+++ b/Engine/Plugins/Animation/IKRig/Source/IKRig/Public/IKRigDefinition.h
@@ -6,10 +6,11 @@
 #include "UObject/Object.h"
 #include "IKRigBoneSetting.h"
 #include "IKRigSkeleton.h"
+#include "IKRigDataTypes.h"
 
 #include "IKRigDefinition.generated.h"
 
-struct FIKRigEffectorGoal;
+
 class UIKRigSolver;
 
 USTRUCT(Blueprintable)
@@ -92,13 +93,26 @@ public:
 	FRetargetDefinition RetargetDefinition;
 
 	/** get a list of all goal names used in the solvers */
-	void GetGoalNamesFromSolvers(TArray<FIKRigEffectorGoal>& OutGoalNames) const;
+	TArray<FIKRigEffectorGoal>& GetEffectorGoals();
 
 	/** get bone associated with a given goal */
-	FName GetBoneNameForGoal(FName GoalName) const;
+	FName GetBoneNameForGoal(const FName& GoalName);
+
+	/** get name of goal at given index. */
+	FName GetGoalName(int32 GoalIndex);
+
+	// UObject interface
+#if WITH_EDITOR
+	virtual void PostEditChangeProperty(struct FPropertyChangedEvent& PropertyChangedEvent) override;
+#endif
+	// end UObject interface
 
 private:
+
+	UPROPERTY(Transient)
+	TArray<FIKRigEffectorGoal> EffectorGoals;
+	bool bEffectorGoalsDirty = true;
+	void UpdateGoalNameArray();
 	
 	friend class UIKRigController;
 };
-
diff --git a/Engine/Plugins/Animation/IKRig/Source/IKRig/Public/IKRigProcessor.h b/Engine/Plugins/Animation/IKRig/Source/IKRig/Public/IKRigProcessor.h
--- a/Engine/Plugins/Animation/IKRig/Source/IKRig/Public/IKRigProcessor.h
+++ b/Engine/Plugins/Animation/IKRig/Source/IKRig/Public/IKRigProcessor.h
@@ -87,6 +87,9 @@ public:
 	
 private:
 
+	/** Update the final pos/rot of all the goals based on their alpha values. */
+	void BlendGoalsByAlpha();
+
 	/** solving disabled until this flag is true */
 	bool bInitialized = false;
 
diff --git a/Engine/Plugins/Animation/IKRig/Source/IKRig/Public/IKRigSolver.h b/Engine/Plugins/Animation/IKRig/Source/IKRig/Public/IKRigSolver.h
--- a/Engine/Plugins/Animation/IKRig/Source/IKRig/Public/IKRigSolver.h
+++ b/Engine/Plugins/Animation/IKRig/Source/IKRig/Public/IKRigSolver.h
@@ -4,6 +4,7 @@
 
 #include "UObject/ObjectMacros.h"
 #include "UObject/Object.h"
+#include "IKRigDataTypes.h"
 #include "IKRigSolver.generated.h"
 
 
@@ -11,20 +12,6 @@ struct FControlRigDrawInterface;
 struct FIKRigGoalContainer;
 struct FIKRigSkeleton;
 
-USTRUCT()
-struct IKRIG_API FIKRigEffectorGoal
-{
-	GENERATED_BODY()
-
-	UPROPERTY(EditAnywhere, Category = FIKRigEffector)
-	FName Goal;
-
-	UPROPERTY(EditAnywhere, Category = FIKRigEffector)
-	FName Bone;
-
-	bool operator==(const FIKRigEffectorGoal& Other) const { return Goal == Other.Goal; }
-};
-
 inline uint32 GetTypeHash(FIKRigEffectorGoal ObjectRef) { return GetTypeHash(ObjectRef.Goal); }
 
 // run time processor 
@@ -47,14 +34,18 @@ public:
 		const FIKRigGoalContainer& Goals,
 		FControlRigDrawInterface* InOutDrawInterface) PURE_VIRTUAL("Solve");
 
-	/** override CollectGoalNames() to tell the processor which goals to collect for this solver
+	/** override AddGoalsInSolver() to tell other systems what goals this rig definition uses
 	 * NOTE: only ADD to OutGoals, do not reset or remove */
-	virtual void CollectGoalNames(TSet<FIKRigEffectorGoal>& OutGoals) const PURE_VIRTUAL("CollectGoalNames");
+	virtual void AddGoalsInSolver(TArray<FIKRigEffectorGoal>& OutGoals) const PURE_VIRTUAL("AddGoalsInSolver");
 
 	#if WITH_EDITOR
 	/** override RenameGoal() when UI renames a goal, you can auto-update any effectors
 	 * that were previously mapped to the old goal name. */
 	virtual void RenameGoal(const FName& OldName, const FName& NewName) {};
 	#endif // WITH_EDITOR
+
+protected:
+	
+	static void AddGoalToArrayNoDuplicates(const FIKRigEffectorGoal& GoalToAdd, TArray<FIKRigEffectorGoal>& OutGoals);
 };
 
diff --git a/Engine/Plugins/Animation/IKRig/Source/IKRig/Public/Retargeter/IKRetargeter.h b/Engine/Plugins/Animation/IKRig/Source/IKRig/Public/Retargeter/IKRetargeter.h
--- a/Engine/Plugins/Animation/IKRig/Source/IKRig/Public/Retargeter/IKRetargeter.h
+++ b/Engine/Plugins/Animation/IKRig/Source/IKRig/Public/Retargeter/IKRetargeter.h
@@ -11,6 +11,53 @@
 
 #include "IKRetargeter.generated.h"
 
+USTRUCT(Blueprintable)
+struct IKRIG_API FIKRetargetChainSettings
+{
+	GENERATED_BODY()
+
+	/** Range -1 to 1. Default 0. Rotates chain away or towards the midline (positive values will raise arms upwards). */
+	UPROPERTY(EditAnywhere, BlueprintReadWrite, Category = Offsets, meta = (ClampMin = "-1.0", ClampMax = "1.0", UIMin = "-1.0", UIMax = "1.0"))
+	float AbductionOffset = 0.0f;
+
+	/** Range -1 to 1. Default 0. Rotates chain forward (+) and backward (-) relative to facing direction. */
+	UPROPERTY(EditAnywhere, BlueprintReadWrite, Category = Offsets, meta = (ClampMin = "-1.0", ClampMax = "1.0", UIMin = "-1.0", UIMax = "1.0"))
+	float SwingOffset = 0.0f;
+
+	/** Range -1 to 1. Default 0. Rotates the chain along an axis from the start to the end of the chain.
+	*  At -1 the chain is rotated laterally 90 degrees. At +1 the chain is rotated medially 90 degrees. */
+	UPROPERTY(EditAnywhere, BlueprintReadWrite, Category = Offsets, meta = (ClampMin = "-1.0", ClampMax = "1.0", UIMin = "-1.0", UIMax = "1.0"))
+	float TwistOffset = 0.0f;
+
+	/** Range -1 to 2. Default 0. Brings IK effector closer (-) or further (+) from origin of chain.
+	*  At -1 the end is on top of the origin. At +2 the end is fully extended twice the length of the chain. */
+	UPROPERTY(EditAnywhere, BlueprintReadWrite, Category = Offsets, meta = (ClampMin = "-1.0", ClampMax = "2.0", UIMin = "-1.0", UIMax = "2.0"))
+	float Extension = 0.0f;
+
+	/** Range 0 to 1. Default 0. Allow the chain to stretch by translating to reach the IK goal locations.
+	*  At 0 the chain will not stretch at all. At 1 the chain will be allowed to stretch double it's full length to reach IK. */
+	UPROPERTY(EditAnywhere, BlueprintReadWrite, Category = Stretch, meta = (ClampMin = "0.0", ClampMax = "1.0", UIMin = "0.0", UIMax = "1.0"))
+	float StretchTolerance = 0.0f;
+
+	/** Range 0 to 1. Default is 1.0. Blend IK effector at the end of this chain towards the original position
+	 *  on the source skeleton (0.0) or the position on the retargeted target skeleton (1.0). */
+	UPROPERTY(EditAnywhere, BlueprintReadWrite, Category = IkMode, meta = (ClampMin = "0.0", ClampMax = "1.0", UIMin = "0.0", UIMax = "1.0"))
+	float IkToSourceOrTarget = 1.0f;
+	
+	/** When true, the source IK position is calculated relative to a source bone and applied relative to a target bone. */
+	UPROPERTY(EditAnywhere, BlueprintReadWrite, Category = IkMode)
+	bool bIkRelativeToSourceBone = false;
+	
+	/** A bone in the SOURCE skeleton that the IK location is relative to. */
+	UPROPERTY(EditAnywhere, BlueprintReadWrite, Category = IkMode, meta = (EditCondition="bIkRelativeToSourceBone"))
+	FName IkRelativeSourceBone;
+	
+	/** A bone in the TARGET skeleton that the IK location is relative to.
+	 * This is usually the same bone as the source skeleton, but may have a different name in the target skeleton. */
+	UPROPERTY(EditAnywhere, BlueprintReadWrite, Category = IkMode, meta = (EditCondition="bIkRelativeToSourceBone"))
+	FName IkRelativeTargetBone;	
+};
+
 struct FRetargetSkeleton
 {
 	TArray<FName> BoneNames;
diff --git a/Engine/Plugins/Animation/IKRig/Source/IKRig/Public/Solvers/ConstraintSolver.h b/Engine/Plugins/Animation/IKRig/Source/IKRig/Public/Solvers/ConstraintSolver.h
deleted file mode 100644
diff --git a/Engine/Plugins/Animation/IKRig/Source/IKRig/Public/Solvers/IKRigFBIKSolver.h b/Engine/Plugins/Animation/IKRig/Source/IKRig/Public/Solvers/IKRigFBIKSolver.h
deleted file mode 100644
diff --git a/Engine/Plugins/Animation/IKRig/Source/IKRig/Public/Solvers/IKRig_PBIKSolver.h b/Engine/Plugins/Animation/IKRig/Source/IKRig/Public/Solvers/IKRig_PBIKSolver.h
new file mode 100644
--- /dev/null
+++ b/Engine/Plugins/Animation/IKRig/Source/IKRig/Public/Solvers/IKRig_PBIKSolver.h
@@ -0,0 +1,57 @@
+// Copyright Epic Games, Inc. All Rights Reserved.
+
+#pragma once
+
+
+#include "IKRigSolver.h"
+#include "Core/PBIKSolver.h"
+#include "PBIK_Shared.h"
+
+#include "IKRig_PBIKSolver.generated.h"
+
+USTRUCT()
+struct FIKRig_PBIKEffector
+{
+	GENERATED_BODY()
+
+	UPROPERTY(EditAnywhere, Category = Settings)
+	FIKRigEffectorGoal Goal;
+
+	UPROPERTY(EditAnywhere, Category = Settings)
+	float StrengthAlpha = 1.0f;
+
+	UPROPERTY(Transient)
+	int32 IndexInSolver = -1;
+};
+
+UCLASS(EditInlineNew, config = Engine, hidecategories = UObject)
+class IKRIG_API UIKRigPBIKSolver : public UIKRigSolver
+{
+	GENERATED_BODY()
+
+public:
+
+	UPROPERTY(EditAnywhere, Category = Settings)
+	FName RootBone;
+	
+	UPROPERTY(EditAnywhere, Category = Settings)
+	TArray<FIKRig_PBIKEffector> Effectors;
+
+	UPROPERTY(EditAnywhere, Category = Settings)
+	TArray<FPBIKBoneSetting> BoneSettings;
+
+	UPROPERTY(EditAnywhere, Category = Settings)
+	FPBIKSolverSettings Settings;
+
+	virtual void Initialize(const FIKRigSkeleton& IKRigSkeleton) override;
+	virtual void Solve(
+		FIKRigSkeleton& IKRigSkeleton,
+		const FIKRigGoalContainer& Goals,
+		FControlRigDrawInterface* InOutDrawInterface) override;
+	virtual void AddGoalsInSolver(TArray<FIKRigEffectorGoal>& OutGoals) const override;
+
+private:
+
+	FPBIKSolver Solver;
+};
+
diff --git a/Engine/Plugins/Animation/IKRig/Source/IKRig/Public/Solvers/IKRig_SetTransform.h b/Engine/Plugins/Animation/IKRig/Source/IKRig/Public/Solvers/IKRig_SetTransform.h
new file mode 100644
--- /dev/null
+++ b/Engine/Plugins/Animation/IKRig/Source/IKRig/Public/Solvers/IKRig_SetTransform.h
@@ -0,0 +1,39 @@
+// Copyright Epic Games, Inc. All Rights Reserved.
+
+#pragma once
+
+#include "IKRigSolver.h"
+
+#include "IKRig_SetTransform.generated.h"
+
+
+UCLASS(EditInlineNew)
+class IKRIG_API UIKRig_SetTransform : public UIKRigSolver
+{
+	GENERATED_BODY()
+
+public:
+	
+	UPROPERTY(EditAnywhere, Category = "Solver")
+	FIKRigEffectorGoal Effector;
+	
+	UPROPERTY(EditAnywhere, Category = "Solver")
+	bool bEnablePosition = true;
+
+	UPROPERTY(EditAnywhere, Category = "Solver")
+	bool bEnableRotation = true;
+
+protected:
+	
+	virtual void Initialize(const FIKRigSkeleton& IKRigSkeleton) override;
+	virtual void Solve(
+		FIKRigSkeleton& IKRigSkeleton, 
+		const FIKRigGoalContainer& Goals,
+		FControlRigDrawInterface* InOutDrawInterface) override;
+	virtual void AddGoalsInSolver(TArray<FIKRigEffectorGoal>& OutGoals) const override;
+
+private:
+	
+	int32 BoneIndex;
+};
+
diff --git a/Engine/Plugins/Animation/IKRig/Source/IKRig/Public/Solvers/TransformSolver.h b/Engine/Plugins/Animation/IKRig/Source/IKRig/Public/Solvers/TransformSolver.h
deleted file mode 100644
diff --git a/Engine/Plugins/Animation/IKRig/Source/IKRigDeveloper/Private/AnimGraphNode_IKRig.cpp b/Engine/Plugins/Animation/IKRig/Source/IKRigDeveloper/Private/AnimGraphNode_IKRig.cpp
--- a/Engine/Plugins/Animation/IKRig/Source/IKRigDeveloper/Private/AnimGraphNode_IKRig.cpp
+++ b/Engine/Plugins/Animation/IKRig/Source/IKRigDeveloper/Private/AnimGraphNode_IKRig.cpp
@@ -71,10 +71,7 @@ void UAnimGraphNode_IKRig::CustomizePinData(UEdGraphPin* Pin, FName SourceProper
 void UAnimGraphNode_IKRig::PostLoad()
 {
 	Super::PostLoad();
-
-	// resize global transforms
 	Node.RebuildGoalList();
-	// I can't reconstruct here
 }
 
 void UAnimGraphNode_IKRig::PreloadRequiredAssets()
diff --git a/Engine/Plugins/Animation/IKRig/Source/IKRigEditor/Private/IKRigController.cpp b/Engine/Plugins/Animation/IKRig/Source/IKRigEditor/Private/IKRigController.cpp
--- a/Engine/Plugins/Animation/IKRig/Source/IKRigEditor/Private/IKRigController.cpp
+++ b/Engine/Plugins/Animation/IKRig/Source/IKRigEditor/Private/IKRigController.cpp
@@ -88,7 +88,7 @@ FIKRigSkeleton& UIKRigController::GetSkeleton() const
 // SOLVERS
 //
 
-UIKRigSolver* UIKRigController::AddSolver(TSubclassOf<UIKRigSolver> InIKRigSolverClass)
+UIKRigSolver* UIKRigController::AddSolver(TSubclassOf<UIKRigSolver> InIKRigSolverClass) const
 {
 	if (!IKRigDefinition)
 	{
@@ -102,6 +102,8 @@ UIKRigSolver* UIKRigController::AddSolver(TSubclassOf<UIKRigSolver> InIKRigSolve
 	check(NewSolver);
 
 	IKRigDefinition->Solvers.Add(NewSolver);
+	IKRigDefinition->bEffectorGoalsDirty = true;
+	
 	return NewSolver;
 }
 
@@ -116,7 +118,7 @@ UIKRigSolver* UIKRigController::GetSolver(int32 Index) const
 	return ValidSolver ? IKRigDefinition->Solvers[Index] : nullptr;
 }
 
-void UIKRigController::RemoveSolver(UIKRigSolver* SolverToDelete)
+void UIKRigController::RemoveSolver(UIKRigSolver* SolverToDelete) const
 {
 	if (!(IKRigDefinition && SolverToDelete))
 	{
@@ -127,6 +129,7 @@ void UIKRigController::RemoveSolver(UIKRigSolver* SolverToDelete)
 	IKRigDefinition->Modify();
 
 	IKRigDefinition->Solvers.Remove(SolverToDelete);
+	IKRigDefinition->bEffectorGoalsDirty = true;
 }
 
 UIKRigBoneSetting* UIKRigController::AddBoneSetting(TSubclassOf<UIKRigBoneSetting> NewSettingType) const
@@ -149,14 +152,14 @@ UIKRigBoneSetting* UIKRigController::AddBoneSetting(TSubclassOf<UIKRigBoneSettin
 // GOALS
 //
 
-void UIKRigController::GetGoalNames(TArray<FIKRigEffectorGoal>& OutGoals) const
+TArray<FIKRigEffectorGoal>* UIKRigController::GetGoalNames() const
 {
 	if (!IKRigDefinition)
 	{
-		return;
+		return nullptr;
 	}
 
-	IKRigDefinition->GetGoalNamesFromSolvers(OutGoals);
+	return &IKRigDefinition->GetEffectorGoals();
 }
 
 void UIKRigController::RenameGoal(const FName& OldName, const FName& NewName) const
@@ -173,6 +176,8 @@ void UIKRigController::RenameGoal(const FName& OldName, const FName& NewName) co
 	{
 		Solver->RenameGoal(OldName, NewName);
 	}
+
+	IKRigDefinition->bEffectorGoalsDirty = true;
 }
 
-#undef LOCTEXT_NAMESPACE
\ No newline at end of file
+#undef LOCTEXT_NAMESPACE
diff --git a/Engine/Plugins/Animation/IKRig/Source/IKRigEditor/Public/IKRigController.h b/Engine/Plugins/Animation/IKRig/Source/IKRigEditor/Public/IKRigController.h
--- a/Engine/Plugins/Animation/IKRig/Source/IKRigEditor/Public/IKRigController.h
+++ b/Engine/Plugins/Animation/IKRig/Source/IKRigEditor/Public/IKRigController.h
@@ -30,8 +30,8 @@ public:
 	FIKRigSkeleton& GetSkeleton() const;
 
 	//** solvers */
-	UIKRigSolver* AddSolver(TSubclassOf<UIKRigSolver> InSolverClass);
-	void RemoveSolver(UIKRigSolver* SolverToDelete);
+	UIKRigSolver* AddSolver(TSubclassOf<UIKRigSolver> InSolverClass) const;
+	void RemoveSolver(UIKRigSolver* SolverToDelete) const;
 	UIKRigSolver* GetSolver(int32 Index) const;
 	int32 GetNumSolvers() const;
 
@@ -40,7 +40,7 @@ public:
 
 	//** goals */
 	void RenameGoal(const FName& OldName, const FName& NewName) const;
-	void GetGoalNames(TArray<FIKRigEffectorGoal>& OutGoals) const;
+	TArray<FIKRigEffectorGoal>* GetGoalNames() const;
 	DECLARE_MULTICAST_DELEGATE(FGoalModified);
 	FGoalModified OnGoalModified;
 
diff --git a/Engine/Plugins/Experimental/FullBodyIK/Source/PBIK/Private/Core/PBIKSolver.cpp b/Engine/Plugins/Experimental/FullBodyIK/Source/PBIK/Private/Core/PBIKSolver.cpp
--- a/Engine/Plugins/Experimental/FullBodyIK/Source/PBIK/Private/Core/PBIKSolver.cpp
+++ b/Engine/Plugins/Experimental/FullBodyIK/Source/PBIK/Private/Core/PBIKSolver.cpp
@@ -17,7 +17,7 @@ namespace PBIK
 	}
 
 	void FEffector::SetGoal(
-		const FVector InPositionGoal,
+		const FVector& InPositionGoal,
 		const FQuat& InRotationGoal,
 		float InTransformAlpha,
 		float InStrengthAlpha)
@@ -38,13 +38,29 @@ namespace PBIK
 		Rotation = FMath::Lerp(RotationOrig, RotationGoal, TransformAlpha);
 		Pin.Pin()->GoalPoint = Position;
 		Pin.Pin()->Alpha = StrengthAlpha;
+
+		// update length of chain this effector controls in the input pose
+		if (ParentSubRoot)
+		{
+			FEffector* ParentEffector = ParentSubRoot->AttachedEffector;
+			const FVector ParentSubRootPosition = ParentEffector ? ParentEffector->PositionOrig : ParentSubRoot->Bone->Position;
+			LengthOfChainInInputPose = (ParentSubRootPosition - Bone->Position).Size();
+		}
 	}
 
 	void FEffector::SquashSubRoots()
 	{
 		// optionally apply a preferred angle to give solver a hint which direction to favor
 		// apply amount of preferred angle proportional to the amount this sub-limb is squashed
-		if (!ParentSubRoot || DistToSubRootOrig <= SMALL_NUMBER)
+
+		// can't squash root chain
+		if (!ParentSubRoot)
+		{
+			return; 
+		}
+
+		// can't squash chain with zero length already
+		if (LengthOfChainInInputPose <= SMALL_NUMBER)
 		{
 			return;
 		}
@@ -54,16 +70,16 @@ namespace PBIK
 		// otherwise use the current position of the FRigidBody
 		FEffector* ParentEffector = ParentSubRoot->AttachedEffector;
 		const FVector ParentSubRootPosition = ParentEffector ? ParentEffector->Position : ParentSubRoot->Position;
-		const float DistToNearestSubRoot = (ParentSubRootPosition - Position).Size();
-		if (DistToNearestSubRoot >= DistToSubRootOrig)
+		const float DistToParentSubRoot = (ParentSubRootPosition - Position).Size();
+		if (DistToParentSubRoot >= LengthOfChainInInputPose)
 		{
 			return; // limb is stretched
 		}
 
 		// shrink distance to reach full blend to preferred angle
-		const float ScaledDistOrig = DistToSubRootOrig * 0.3f;
+		const float ScaledDistOrig = LengthOfChainInInputPose * 0.3f;
 		// amount squashed (clamped to scaled original length)
-		float DeltaSquash = DistToSubRootOrig - DistToNearestSubRoot;
+		float DeltaSquash = LengthOfChainInInputPose - DistToParentSubRoot;
 		DeltaSquash = DeltaSquash > ScaledDistOrig ? ScaledDistOrig : DeltaSquash;
 		const float SquashPercent = DeltaSquash / ScaledDistOrig;
 		if (SquashPercent < 0.01f)
@@ -139,7 +155,7 @@ void FPBIKSolver::Solve(const FPBIKSolverSettings& Settings)
 	{
 		Effector.UpdateFromInputs();
 	}
-
+	
 	// squash sub-roots to apply preferred angles
 	for (FEffector& Effector : Effectors)
 	{
@@ -234,7 +250,10 @@ bool FPBIKSolver::Initialize()
 		return false;
 	}
 
-	InitConstraints();
+	if (!InitConstraints())
+	{
+		return false;
+	}
 
 	bReadyToSimulate = true;
 
@@ -246,13 +265,15 @@ bool FPBIKSolver::InitBones()
 	using PBIK::FEffector;
 	using PBIK::FBone;
 
-	if (!ensureMsgf(Bones.Num() > 0, TEXT("PBIK: no bones added to solver. Cannot initialize.")))
+	if (Bones.IsEmpty())
 	{
+		UE_LOG(LogTemp, Warning, TEXT("PBIK: no bones added to solver. Cannot initialize."));
 		return false;
 	}
 
-	if (!ensureMsgf(Effectors.Num() > 0, TEXT("PBIK: no effectors added to solver. Cannot initialize.")))
+	if (Effectors.IsEmpty())
 	{
+		UE_LOG(LogTemp, Warning, TEXT("PBIK: no effectors added to solver. Cannot initialize."));
 		return false;
 	}
 
@@ -267,13 +288,15 @@ bool FPBIKSolver::InitBones()
 		}
 	}
 
-	if (!ensureMsgf(SolverRoot, TEXT("PBIK: root bone not set. Cannot initialize.")))
+	if (!SolverRoot)
 	{
+		UE_LOG(LogTemp, Warning, TEXT("PBIK: root bone not set or not found. Cannot initialize."));
 		return false;
 	}
 
-	if (!ensureMsgf(NumSolverRoots == 1, TEXT("PBIK: more than 1 bone was marked as solver root. Cannot initialize.")))
+	if (NumSolverRoots > 1)
 	{
+		UE_LOG(LogTemp, Warning, TEXT("PBIK: more than 1 bone was marked as solver root. Cannot initialize."));
 		return false;
 	}
 
@@ -288,9 +311,9 @@ bool FPBIKSolver::InitBones()
 
 		// validate parent
 		const bool bIndexInRange = Bone.ParentIndex >= 0 && Bone.ParentIndex < Bones.Num();
-		if (!ensureMsgf(bIndexInRange,
-			TEXT("PBIK: bone found with invalid parent index. Cannot initialize.")))
+		if (!bIndexInRange)
 		{
+			UE_LOG(LogTemp, Warning, TEXT("PBIK: bone found with invalid parent index. Cannot initialize."));
 			return false;
 		}
 
@@ -350,8 +373,9 @@ bool FPBIKSolver::InitBodies()
 		while (true)
 		{
 			FBone* BodyBone = NextBone->bIsSolverRoot ? NextBone : NextBone->Parent;
-			if (!ensureMsgf(BodyBone, TEXT("PBIK: effector is on bone that is not on or below root bone.")))
+			if (!BodyBone)
 			{
+				UE_LOG(LogTemp, Warning, TEXT("PBIK: effector is on bone that is not on or below root bone."));
 				return false;
 			}
 
@@ -396,7 +420,6 @@ bool FPBIKSolver::InitBodies()
 			if (Parent->bIsSubRoot || Parent->bIsSolverRoot)
 			{
 				Effector.ParentSubRoot = Parent->Body;
-				Effector.DistToSubRootOrig = (Parent->Body->Position - Effector.Bone->Position).Size();
 				break;
 			}
 
@@ -419,7 +442,7 @@ void FPBIKSolver::AddBodyForBone(PBIK::FBone* Bone)
 	Bodies.Emplace(Bone);
 }
 
-void FPBIKSolver::InitConstraints()
+bool FPBIKSolver::InitConstraints()
 {
 	using PBIK::FEffector;
 	using PBIK::FBone;
@@ -433,9 +456,10 @@ void FPBIKSolver::InitConstraints()
 	for (FEffector& Effector : Effectors)
 	{
 		FBone* BodyBone = Effector.Bone->bIsSolverRoot ? Effector.Bone : Effector.Bone->Parent;
-		if (!ensureMsgf(BodyBone, TEXT("PBIK: effector is on bone that does not have a parent.")))
+		if (!BodyBone)
 		{
-			return;
+			UE_LOG(LogTemp, Warning, TEXT("PBIK: effector is on bone that does not have a parent."));
+			return false;
 		}
 
 		FRigidBody* Body = BodyBone->Body;
@@ -466,6 +490,8 @@ void FPBIKSolver::InitConstraints()
 		TSharedPtr<FJointConstraint> Constraint = MakeShared<FJointConstraint>(ParentBody, &Body);
 		Constraints.Add(Constraint);
 	}
+
+	return true;
 }
 
 PBIK::FDebugDraw* FPBIKSolver::GetDebugDraw()
diff --git a/Engine/Plugins/Experimental/FullBodyIK/Source/PBIK/Public/Core/PBIKSolver.h b/Engine/Plugins/Experimental/FullBodyIK/Source/PBIK/Public/Core/PBIKSolver.h
--- a/Engine/Plugins/Experimental/FullBodyIK/Source/PBIK/Public/Core/PBIKSolver.h
+++ b/Engine/Plugins/Experimental/FullBodyIK/Source/PBIK/Public/Core/PBIKSolver.h
@@ -30,16 +30,19 @@ struct FEffector
 
 	FBone* Bone;
 	TWeakPtr<FPinConstraint> Pin;
-	
-	float DistToSubRootOrig;
 	FRigidBody* ParentSubRoot = nullptr;
+	float LengthOfChainInInputPose;
 
 	float TransformAlpha;
 	float StrengthAlpha;
 
 	FEffector(FBone* InBone);
 
-	void SetGoal(const FVector InPositionGoal, const FQuat& InRotationGoal, float InTransformAlpha, float InStrengthAlpha);
+	void SetGoal(
+		const FVector& InPositionGoal,
+		const FQuat& InRotationGoal,
+		float InTransformAlpha,
+		float InStrengthAlpha);
 
 	void UpdateFromInputs();
 
@@ -49,7 +52,7 @@ struct FEffector
 } // namespace
 
 USTRUCT()
-struct FPBIKSolverSettings
+struct PBIK_API FPBIKSolverSettings
 {
 	GENERATED_BODY()
 
@@ -69,7 +72,7 @@ struct FPBIKSolverSettings
 
 
 USTRUCT()
-struct FPBIKSolver
+struct PBIK_API FPBIKSolver
 {
 	GENERATED_BODY()
 
@@ -129,7 +132,7 @@ private:
 
 	bool InitBodies();
 
-	void InitConstraints();
+	bool InitConstraints();
 
 	void AddBodyForBone(PBIK::FBone* Bone);
 
diff --git a/Engine/Plugins/Experimental/FullBodyIK/Source/PBIK/Public/PBIK_Shared.h b/Engine/Plugins/Experimental/FullBodyIK/Source/PBIK/Public/PBIK_Shared.h
--- a/Engine/Plugins/Experimental/FullBodyIK/Source/PBIK/Public/PBIK_Shared.h
+++ b/Engine/Plugins/Experimental/FullBodyIK/Source/PBIK/Public/PBIK_Shared.h
@@ -17,7 +17,7 @@ enum class EPBIKLimitType : uint8
 };
 
 USTRUCT(BlueprintType)
-struct FPBIKBoneSetting
+struct PBIK_API FPBIKBoneSetting
 {
 	GENERATED_BODY()
 
@@ -28,38 +28,38 @@ struct FPBIKBoneSetting
 		Z(EPBIKLimitType::Free),
 		PreferredAngles(FVector::ZeroVector){}
 
-	UPROPERTY(meta = (Constant, CustomWidget = "BoneName"))
+	UPROPERTY(EditAnywhere, Category = Bone, meta = (Constant, CustomWidget = "BoneName"))
 	FName Bone;
 
-	UPROPERTY(meta = (ClampMin = "0", ClampMax = "1", UIMin = "0.0", UIMax = "1.0"))
+	UPROPERTY(EditAnywhere, Category = Stiffness, meta = (ClampMin = "0", ClampMax = "1", UIMin = "0.0", UIMax = "1.0"))
 	float RotationStiffness = 0.0f;
-	UPROPERTY(meta = (ClampMin = "0", ClampMax = "1", UIMin = "0.0", UIMax = "1.0"))
+	UPROPERTY(EditAnywhere, Category = Stiffness, meta = (ClampMin = "0", ClampMax = "1", UIMin = "0.0", UIMax = "1.0"))
 	float PositionStiffness = 0.0f;
 
-	UPROPERTY()
+	UPROPERTY(EditAnywhere, Category = Limits)
 	EPBIKLimitType X;
-	UPROPERTY(meta = (ClampMin = "-180", ClampMax = "0", UIMin = "-180.0", UIMax = "0.0"))
+	UPROPERTY(EditAnywhere, Category = Limits, meta = (ClampMin = "-180", ClampMax = "0", UIMin = "-180.0", UIMax = "0.0"))
 	float MinX = 0.0f;
-	UPROPERTY(meta = (ClampMin = "0", ClampMax = "180", UIMin = "0.0", UIMax = "180.0"))
+	UPROPERTY(EditAnywhere, Category = Limits, meta = (ClampMin = "0", ClampMax = "180", UIMin = "0.0", UIMax = "180.0"))
 	float MaxX = 0.0f;
 
-	UPROPERTY()
+	UPROPERTY(EditAnywhere, Category = Limits)
 	EPBIKLimitType Y;
-	UPROPERTY(meta = (ClampMin = "-180", ClampMax = "0", UIMin = "-180.0", UIMax = "0.0"))
+	UPROPERTY(EditAnywhere, Category = Limits, meta = (ClampMin = "-180", ClampMax = "0", UIMin = "-180.0", UIMax = "0.0"))
 	float MinY = 0.0f;
-	UPROPERTY(meta = (ClampMin = "0", ClampMax = "180", UIMin = "0.0", UIMax = "180.0"))
+	UPROPERTY(EditAnywhere, Category = Limits, meta = (ClampMin = "0", ClampMax = "180", UIMin = "0.0", UIMax = "180.0"))
 	float MaxY = 0.0f;
 
-	UPROPERTY()
+	UPROPERTY(EditAnywhere, Category = Limits)
 	EPBIKLimitType Z;
-	UPROPERTY(meta = (ClampMin = "-180", ClampMax = "0", UIMin = "-180.0", UIMax = "0.0"))
+	UPROPERTY(EditAnywhere, Category = Limits, meta = (ClampMin = "-180", ClampMax = "0", UIMin = "-180.0", UIMax = "0.0"))
 	float MinZ = 0.0f;
-	UPROPERTY(meta = (ClampMin = "0", ClampMax = "180", UIMin = "0.0", UIMax = "180.0"))
+	UPROPERTY(EditAnywhere, Category = Limits, meta = (ClampMin = "0", ClampMax = "180", UIMin = "0.0", UIMax = "180.0"))
 	float MaxZ = 0.0f;
 
-	UPROPERTY()
+	UPROPERTY(EditAnywhere, Category = PreferredAngles)
 	bool bUsePreferredAngles = false;
-	UPROPERTY()
+	UPROPERTY(EditAnywhere, Category = PreferredAngles)
 	FVector PreferredAngles;
 
 	void CopyToCoreStruct(PBIK::FBoneSettings& Settings) const
