PBIK bug fix: root bone constraints do not work.

diff --git a/Engine/Plugins/Experimental/FullBodyIK/Source/PBIK/Private/Core/PBIKBody.cpp b/Engine/Plugins/Experimental/FullBodyIK/Source/PBIK/Private/Core/PBIKBody.cpp
--- a/Engine/Plugins/Experimental/FullBodyIK/Source/PBIK/Private/Core/PBIKBody.cpp
+++ b/Engine/Plugins/Experimental/FullBodyIK/Source/PBIK/Private/Core/PBIKBody.cpp
@@ -3,6 +3,8 @@
 #include "Core/PBIKBody.h"
 #include "Core/PBIKSolver.h"
 
+//#pragma optimize("", off)
+
 namespace PBIK
 {
 
@@ -73,9 +75,12 @@ void FRigidBody::Initialize(FBone* SolverRoot)
 
 void FRigidBody::UpdateFromInputs(const FPBIKSolverSettings& Settings)
 {
-	// set to input pose
-	Position = Bone->Position - Bone->Rotation * BoneLocalPosition;
-	Rotation = Bone->Rotation;
+	if (Settings.bStartSolveFromInputPose)
+	{
+		// set to input pose
+		Position = Bone->Position - Bone->Rotation * BoneLocalPosition;
+		Rotation = Bone->Rotation;
+	}
 
 	// Body.Length used as rough approximation of the mass of the body
 	// for fork joints (multiple solved children) we sum lengths to all children (see Initialize)
@@ -99,24 +104,20 @@ FRigidBody* FRigidBody::GetParentBody() const
 
 void FRigidBody::ApplyPushToRotateBody(const FVector& Push, const FVector& Offset)
 {
+	if (Pin && Pin->bEnabled && Pin->bPinRotation)
+	{
+		return; // rotation of this body is pinned
+	}
+	
 	// equation 8 in "Detailed Rigid Body Simulation with XPBD"
 	FVector Omega = InvMass * (1.0f - J.RotationStiffness) * FVector::CrossProduct(Offset, Push);
 	FQuat OQ(Omega.X, Omega.Y, Omega.Z, 0.0f);
-	OQ = OQ * Rotation;
-	OQ.X *= 0.5f;
-	OQ.Y *= 0.5f;
-	OQ.Z *= 0.5f;
-	OQ.W *= 0.5f;
-	Rotation.X = Rotation.X + OQ.X;
-	Rotation.Y = Rotation.Y + OQ.Y;
-	Rotation.Z = Rotation.Z + OQ.Z;
-	Rotation.W = Rotation.W + OQ.W;
-	Rotation.Normalize();
+	ApplyRotationDelta(OQ, false);
 }
 
 void FRigidBody::ApplyPushToPosition(const FVector& Push)
 {
-	if (AttachedEffector)
+	if (Pin && Pin->bEnabled)
 	{
 		return; // pins are locked
 	}
@@ -124,4 +125,23 @@ void FRigidBody::ApplyPushToPosition(const FVector& Push)
 	Position += Push * (1.0f - J.PositionStiffness);
 }
 
+void FRigidBody::ApplyRotationDelta(const FQuat& InDelta, const bool bNegated)
+{
+	if (Pin && Pin->bEnabled && Pin->bPinRotation)
+	{
+		return; // rotation of this body is pinned
+	}
+
+	/** InDelta is assumed to be a "pure" quaternion representing an infintesimal rotation */
+	FQuat Delta = InDelta * Rotation;
+	Delta.X *= 0.5f;
+	Delta.Y *= 0.5f;
+	Delta.Z *= 0.5f;
+	Delta.W *= 0.5f;
+	Rotation.X = Rotation.X + (bNegated ? -Delta.X : Delta.X);
+	Rotation.Y = Rotation.Y + (bNegated ? -Delta.Y : Delta.Y);
+	Rotation.Z = Rotation.Z + (bNegated ? -Delta.Z : Delta.Z);
+	Rotation.W = Rotation.W + (bNegated ? -Delta.W : Delta.W);
+	Rotation.Normalize();
+}
 } // namespace
\ No newline at end of file
diff --git a/Engine/Plugins/Experimental/FullBodyIK/Source/PBIK/Private/Core/PBIKConstraint.cpp b/Engine/Plugins/Experimental/FullBodyIK/Source/PBIK/Private/Core/PBIKConstraint.cpp
--- a/Engine/Plugins/Experimental/FullBodyIK/Source/PBIK/Private/Core/PBIKConstraint.cpp
+++ b/Engine/Plugins/Experimental/FullBodyIK/Source/PBIK/Private/Core/PBIKConstraint.cpp
@@ -4,6 +4,8 @@
 #include "Core/PBIKSolver.h"
 #include "Math/UnrealMathUtility.h"
 
+//#pragma optimize("", off)
+
 namespace PBIK
 {
 	
@@ -40,13 +42,13 @@ void FJointConstraint::Solve(bool bMoveSubRoots)
 
 	// calc inv mass of body A
 	float AInvMass = 1.0f;
-	AInvMass -= A->AttachedEffector ? A->AttachedEffector->StrengthAlpha : 0.0f;
+	AInvMass -= (A->Pin && A->Pin->bEnabled) ? A->Pin->Alpha : 0.0f;
 	AInvMass -= !bMoveSubRoots && A->Bone->bIsSubRoot ? 1.0f : 0.0f;
 	AInvMass = AInvMass <= 0.0f ? 0.0f : AInvMass;
 
 	// calc inv mass of body B
 	float BInvMass = 1.0f;
-	BInvMass -= B->AttachedEffector ? B->AttachedEffector->StrengthAlpha : 0.0f;
+	BInvMass -= (B->Pin && B->Pin->bEnabled) ? B->Pin->Alpha : 0.0f;
 	BInvMass -= !bMoveSubRoots && B->Bone->bIsSubRoot ? 1.0f : 0.0f;
 	BInvMass = BInvMass <= 0.0f ? 0.0f : BInvMass;
 
@@ -88,30 +90,8 @@ FVector FJointConstraint::GetPositionCorrection(FVector& OutBodyToA, FVector& Ou
 
 void FJointConstraint::ApplyRotationCorrection(FQuat PureRotA, FQuat PureRotB) const
 {
-	// https://matthias-research.github.io/pages/publications/PBDBodies.pdf/
-	// Equation 8 and 9 from "Detailed Rigid Body Simulation with Extended Position Based Dynamics"
-
-	PureRotA.X *= 0.5f;
-	PureRotA.Y *= 0.5f;
-	PureRotA.Z *= 0.5f;
-	PureRotB.X *= 0.5f;
-	PureRotB.Y *= 0.5f;
-	PureRotB.Z *= 0.5f;
-	PureRotA = PureRotA * A->Rotation;
-	PureRotB = PureRotB * B->Rotation;
-
-	A->Rotation.X = A->Rotation.X + PureRotA.X;
-	A->Rotation.Y = A->Rotation.Y + PureRotA.Y;
-	A->Rotation.Z = A->Rotation.Z + PureRotA.Z;
-	A->Rotation.W = A->Rotation.W + PureRotA.W;
-
-	B->Rotation.X = B->Rotation.X - PureRotB.X;
-	B->Rotation.Y = B->Rotation.Y - PureRotB.Y;
-	B->Rotation.Z = B->Rotation.Z - PureRotB.Z;
-	B->Rotation.W = B->Rotation.W - PureRotB.W;
-
-	A->Rotation.Normalize();
-	B->Rotation.Normalize();
+	A->ApplyRotationDelta(PureRotA, false);
+	B->ApplyRotationDelta(PureRotB, true);
 }
 
 void FJointConstraint::UpdateJointLimits()
@@ -287,11 +267,20 @@ float FJointConstraint::SignedAngleBetweenNormals(
 	return Dot >= 0 ? Angle : -Angle;
 }
 	
-FPinConstraint::FPinConstraint(FRigidBody* InBody, const FVector& InPinPoint)
+FPinConstraint::FPinConstraint(
+	FRigidBody* InBody,
+	const FVector& InPinPositionOrig,
+	const FQuat& InPinRotationOrig,
+	const bool bInPinRotation)
 {
+	GoalPosition = InPinPositionOrig;
+	GoalRotation = InPinRotationOrig;
+	
 	A = InBody;
-	PinPointLocalToA = A->Rotation.Inverse() * (InPinPoint - A->Position);
-	GoalPoint = InPinPoint;
+	PinPointLocalToA = A->Rotation.Inverse() * (GoalPosition - A->Position);
+
+	ARotLocalToPin = A->Rotation * GoalRotation.Inverse();
+	bPinRotation = bInPinRotation;
 }
 
 void FPinConstraint::Solve(bool bMoveSubRoots)
@@ -310,21 +299,37 @@ void FPinConstraint::Solve(bool bMoveSubRoots)
 	FVector AToPinPoint;
 	FVector Correction = GetPositionCorrection(AToPinPoint);
 
-	// rotate body from alignment of pin points
-	A->ApplyPushToRotateBody(Correction, AToPinPoint);
+	if (bPinRotation)
+	{
+		// move body to pin location
+		A->Position += Correction;
+		// keep body at fixed rotation relative to the pin goal rotation
+		A->Rotation = ARotLocalToPin * GoalRotation;
+	}else
+	{
+		// rotate body from alignment of pin points
+		A->ApplyPushToRotateBody(Correction, AToPinPoint);
+
+		// apply positional correction to Body to align with target (after rotation)
+		// (applying directly without considering PositionStiffness because PinConstraints need
+		// to precisely pull the attached body to achieve convergence)
+		Correction = GetPositionCorrection(AToPinPoint);
+		A->Position += Correction; 
+	}
+}
 
-	// apply positional correction to Body to align with target (after rotation)
-	// (applying directly without considering PositionStiffness because PinConstraints need
-	// to precisely pull the attached body to achieve convergence)
-	Correction = GetPositionCorrection(AToPinPoint);
-	A->Position += Correction; 
+void FPinConstraint::SetGoal(const FVector& InGoalPosition, const FQuat& InGoalRotation, const float InAlpha)
+{
+	GoalPosition = InGoalPosition;
+	GoalRotation = InGoalRotation;
+	Alpha = InAlpha;
 }
 
 FVector FPinConstraint::GetPositionCorrection(FVector& OutBodyToPinPoint) const
 {
 	OutBodyToPinPoint = A->Rotation * PinPointLocalToA;
 	const FVector PinPoint = A->Position + OutBodyToPinPoint;
-	return (GoalPoint - PinPoint) * Alpha;
+	return (GoalPosition - PinPoint) * Alpha;
 }
 
 } // namespace
diff --git a/Engine/Plugins/Experimental/FullBodyIK/Source/PBIK/Private/Core/PBIKSolver.cpp b/Engine/Plugins/Experimental/FullBodyIK/Source/PBIK/Private/Core/PBIKSolver.cpp
--- a/Engine/Plugins/Experimental/FullBodyIK/Source/PBIK/Private/Core/PBIKSolver.cpp
+++ b/Engine/Plugins/Experimental/FullBodyIK/Source/PBIK/Private/Core/PBIKSolver.cpp
@@ -36,13 +36,12 @@ namespace PBIK
 	{
 		Position = FMath::Lerp(PositionOrig, PositionGoal, TransformAlpha);
 		Rotation = FMath::Lerp(RotationOrig, RotationGoal, TransformAlpha);
-		Pin.Pin()->GoalPoint = Position;
-		Pin.Pin()->Alpha = StrengthAlpha;
+		Pin.Pin()->SetGoal(Position, Rotation, StrengthAlpha);
 
 		// update length of chain this effector controls in the input pose
 		if (ParentSubRoot)
 		{
-			FEffector* ParentEffector = ParentSubRoot->AttachedEffector;
+			FEffector* ParentEffector = ParentSubRoot->Effector;
 			const FVector ParentSubRootPosition = ParentEffector ? ParentEffector->PositionOrig : ParentSubRoot->Bone->Position;
 			LengthOfChainInInputPose = (ParentSubRootPosition - Bone->Position).Size();
 		}
@@ -66,9 +65,9 @@ namespace PBIK
 		}
 
 		// we have to be careful here when calculating the distance to the parent sub-root.
-		// if the the parent sub-root is attached to an effector, use the effector's position
+		// if the parent sub-root is attached to an effector, use the effector's position
 		// otherwise use the current position of the FRigidBody
-		FEffector* ParentEffector = ParentSubRoot->AttachedEffector;
+		FEffector* ParentEffector = ParentSubRoot->Effector;
 		const FVector ParentSubRootPosition = ParentEffector ? ParentEffector->Position : ParentSubRoot->Position;
 		const float DistToParentSubRoot = (ParentSubRootPosition - Position).Size();
 		if (DistToParentSubRoot >= LengthOfChainInInputPose)
@@ -125,7 +124,7 @@ void FPBIKSolver::Solve(const FPBIKSolverSettings& Settings)
 
 	// initialize local bone transforms
 	// this has to be done every tick because incoming animation can modify these
-	// even the LocalPosition has to be updated incase translation is animated
+	// even the LocalPosition has to be updated in case translation is animated
 	for (FBone& Bone : Bones)
 	{
 		FBone* Parent = Bone.Parent;
@@ -463,19 +462,21 @@ bool FPBIKSolver::InitConstraints()
 		}
 
 		FRigidBody* Body = BodyBone->Body;
-		TSharedPtr<FPinConstraint> Constraint = MakeShared<FPinConstraint>(Body, Effector.Position);
+		TSharedPtr<FPinConstraint> Constraint = MakeShared<FPinConstraint>(Body, Effector.Position, Effector.Rotation, false);
 		Effector.Pin = Constraint;
-		Body->AttachedEffector = &Effector;
+		Body->Effector = &Effector;
+		Body->Pin = Constraint.Get();
 		Constraints.Add(Constraint);
 	}
 
 	// pin root body to animated location 
 	// this constraint is by default off in solver settings
-	if (!SolverRoot->Body->AttachedEffector) // only add if user hasn't added their own root effector
+	if (!SolverRoot->Body->Effector) // only add if user hasn't added their own root effector
 	{
-		const TSharedPtr<FPinConstraint> RootConstraint = MakeShared<FPinConstraint>(SolverRoot->Body, SolverRoot->Position);
+		const TSharedPtr<FPinConstraint> RootConstraint = MakeShared<FPinConstraint>(SolverRoot->Body, SolverRoot->Position, SolverRoot->Rotation, true);
 		Constraints.Add(RootConstraint);
 		RootPin = RootConstraint;
+		SolverRoot->Body->Pin = RootConstraint.Get();
 	}
 
 	// constrain all bodies together (child to parent)
diff --git a/Engine/Plugins/Experimental/FullBodyIK/Source/PBIK/Public/Core/PBIKBody.h b/Engine/Plugins/Experimental/FullBodyIK/Source/PBIK/Public/Core/PBIKBody.h
--- a/Engine/Plugins/Experimental/FullBodyIK/Source/PBIK/Public/Core/PBIKBody.h
+++ b/Engine/Plugins/Experimental/FullBodyIK/Source/PBIK/Public/Core/PBIKBody.h
@@ -10,10 +10,10 @@ struct FPBIKSolverSettings;
 
 namespace PBIK
 {
-
-struct FJointConstraint;
-struct FRigidBody;
-struct FEffector;
+	struct FPinConstraint;
+	struct FJointConstraint;
+	struct FRigidBody;
+	struct FEffector;
 
 struct FBone
 {
@@ -75,6 +75,8 @@ struct FRigidBody
 {
 	FBone* Bone = nullptr;
 	FBoneSettings J;
+	FPinConstraint* Pin = nullptr;
+	FEffector* Effector = nullptr;
 
 	FVector Position;
 	FQuat Rotation;
@@ -83,7 +85,6 @@ struct FRigidBody
 	TArray<FVector> ChildLocalPositions;
 
 	float InvMass = 0.0f;
-	FEffector* AttachedEffector = nullptr;
 	float Length;
 	
 private:
@@ -105,6 +106,8 @@ public:
 	void ApplyPushToRotateBody(const FVector& Push, const FVector& Offset);
 	
 	void ApplyPushToPosition(const FVector& Push);
+
+	void ApplyRotationDelta(const FQuat& InDelta, const bool bNegated);
 };
 
 // for sorting Bodies hierarchically (root to leaf order)
diff --git a/Engine/Plugins/Experimental/FullBodyIK/Source/PBIK/Public/Core/PBIKConstraint.h b/Engine/Plugins/Experimental/FullBodyIK/Source/PBIK/Public/Core/PBIKConstraint.h
--- a/Engine/Plugins/Experimental/FullBodyIK/Source/PBIK/Public/Core/PBIKConstraint.h
+++ b/Engine/Plugins/Experimental/FullBodyIK/Source/PBIK/Public/Core/PBIKConstraint.h
@@ -83,22 +83,37 @@ private:
 
 struct FPinConstraint : public FConstraint
 {
-	FVector GoalPoint;
+private:
+	
+	FVector GoalPosition;
+	FQuat GoalRotation;
+	
 	FRigidBody* A;
 	FVector PinPointLocalToA;
 	float Alpha = 1.0;
+	FQuat ARotLocalToPin;
+	bool bPinRotation = false;
 
 public:
-
-	FPinConstraint(FRigidBody* InBody, const FVector& InPinPoint);
+	
+	FPinConstraint(
+		FRigidBody* InBody,
+		const FVector& InPinPositionOrig,
+		const FQuat& InPinRotationOrig,
+		const bool bInPinRotation);
 
 	virtual ~FPinConstraint() {};
 
 	virtual void Solve(bool bMoveSubRoots) override;
 
+	void SetGoal(const FVector& InGoalPosition, const FQuat& InGoalRotation, const float InAlpha);
+
 private:
 
 	FVector GetPositionCorrection(FVector& OutBodyToPinPoint) const;
+
+	friend FRigidBody;
+	friend FJointConstraint;
 };
 
 } // namespace
diff --git a/Engine/Plugins/Experimental/FullBodyIK/Source/PBIK/Public/Core/PBIKSolver.h b/Engine/Plugins/Experimental/FullBodyIK/Source/PBIK/Public/Core/PBIKSolver.h
--- a/Engine/Plugins/Experimental/FullBodyIK/Source/PBIK/Public/Core/PBIKSolver.h
+++ b/Engine/Plugins/Experimental/FullBodyIK/Source/PBIK/Public/Core/PBIKSolver.h
@@ -56,20 +56,26 @@ struct PBIK_API FPBIKSolverSettings
 {
 	GENERATED_BODY()
 
+	/** High iteration counts can help solve complex joint configurations with competing constraints, but will increase runtime cost. Default is 20. */
 	UPROPERTY(EditAnywhere, Category = SolverSettings, meta = (ClampMin = "0", UIMin = "0.0", UIMax = "200.0"))
 	int32 Iterations = 20;
 
+	/** A global mass multiplier; higher values will make the joints more stiff, but require more iterations. Typical range is 0.0 to 10.0. */
 	UPROPERTY(EditAnywhere, Category = SolverSettings, meta = (ClampMin = "0", UIMin = "0.0", UIMax = "10.0"))
 	float MassMultiplier = 1.0f;
 
+	/** If true, joints will translate to reach the effectors; causing bones to lengthen if necessary. Good for cartoon effects. Default is false. */
 	UPROPERTY(EditAnywhere, Category = SolverSettings)
 	bool bAllowStretch = false;
 
+	/** Lock the position and rotation of the solver root bone in-place (at animated position). Useful for partial-body solves. Default is false. */
 	UPROPERTY(EditAnywhere, Category = SolverSettings)
 	bool bPinRoot = false;
-};
-
 
+	/** When true, the solver is reset each tick to start from the current input pose. If false, incoming animated poses are ignored and the solver starts from the results of the previous solve. Default is true. */
+	UPROPERTY(EditAnywhere, Category = SolverSettings)
+	bool bStartSolveFromInputPose = true;
+};
 
 USTRUCT()
 struct PBIK_API FPBIKSolver
