PBIK solver upgrades: fixed hinge constraint bug, added gradual mass falloff over iterations

diff --git a/Engine/Plugins/Animation/IKRig/Source/IKRig/Private/Solvers/IKRig_PBIKSolver.cpp b/Engine/Plugins/Animation/IKRig/Source/IKRig/Private/Solvers/IKRig_PBIKSolver.cpp
--- a/Engine/Plugins/Animation/IKRig/Source/IKRig/Private/Solvers/IKRig_PBIKSolver.cpp
+++ b/Engine/Plugins/Animation/IKRig/Source/IKRig/Private/Solvers/IKRig_PBIKSolver.cpp
@@ -104,6 +104,7 @@ void UIKRigPBIKSolver::Solve(FIKRigSkeleton& IKRigSkeleton, const FIKRigGoalCont
 	Settings.Iterations = Iterations;
 	Settings.bAllowStretch = bAllowStretch;
 	Settings.MassMultiplier = MassMultiplier;
+	Settings.MinMassMultiplier = MinMassMultiplier;
 	Settings.bPinRoot = bPinRoot;
 	Settings.bStartSolveFromInputPose = bStartSolveFromInputPose;
 
diff --git a/Engine/Plugins/Experimental/FullBodyIK/Source/PBIK/Private/Core/PBIKBody.cpp b/Engine/Plugins/Experimental/FullBodyIK/Source/PBIK/Private/Core/PBIKBody.cpp
--- a/Engine/Plugins/Experimental/FullBodyIK/Source/PBIK/Private/Core/PBIKBody.cpp
+++ b/Engine/Plugins/Experimental/FullBodyIK/Source/PBIK/Private/Core/PBIKBody.cpp
@@ -84,7 +84,8 @@ void FRigidBody::UpdateFromInputs(const FPBIKSolverSettings& Settings)
 
 	// Body.Length used as rough approximation of the mass of the body
 	// for fork joints (multiple solved children) we sum lengths to all children (see Initialize)
-	InvMass = 1.0f / ( Length * ((Settings.MassMultiplier * GLOBAL_UNITS) + 0.5f));
+	MaxInvMass = 1.0f / ( Length * ((Settings.MassMultiplier * GLOBAL_UNITS) + 0.5f));
+	MinInvMass = 1.0f / ( Length * ((Settings.MinMassMultiplier * GLOBAL_UNITS) + 0.5f));
 }
 
 int FRigidBody::GetNumBonesToRoot() const
diff --git a/Engine/Plugins/Experimental/FullBodyIK/Source/PBIK/Private/Core/PBIKConstraint.cpp b/Engine/Plugins/Experimental/FullBodyIK/Source/PBIK/Private/Core/PBIKConstraint.cpp
--- a/Engine/Plugins/Experimental/FullBodyIK/Source/PBIK/Private/Core/PBIKConstraint.cpp
+++ b/Engine/Plugins/Experimental/FullBodyIK/Source/PBIK/Private/Core/PBIKConstraint.cpp
@@ -4,8 +4,6 @@
 #include "Core/PBIKSolver.h"
 #include "Math/UnrealMathUtility.h"
 
-//#pragma optimize("", off)
-
 namespace PBIK
 {
 	
@@ -125,43 +123,17 @@ void FJointConstraint::UpdateJointLimits()
 	if (bXHinge)
 	{
 		UpdateLocalRotateAxes(true, false, false);
-		const FVector CrossProd = FVector::CrossProduct(XA, XB);
-		const FQuat PureRotA = FQuat(CrossProd.X, CrossProd.Y, CrossProd.Z, 0.0f);
-		ApplyRotationCorrection(PureRotA, PureRotA);
+		RotateToAlignAxes(XA, XB);
 	}
 	else if (bYHinge)
 	{
 		UpdateLocalRotateAxes(false, true, false);
-		const FVector CrossProd = FVector::CrossProduct(YA, YB);
-		const FQuat PureRotA = FQuat(CrossProd.X, CrossProd.Y, CrossProd.Z, 0.0f);
-		ApplyRotationCorrection(PureRotA, PureRotA);
+		RotateToAlignAxes(YA, YB);
 	}
 	else if (bZHinge)
 	{
 		UpdateLocalRotateAxes(false, false, true);
-		const FVector CrossProd = FVector::CrossProduct(ZA, ZB);
-		const FQuat PureRotA = FQuat(CrossProd.X, CrossProd.Y, CrossProd.Z, 0.0f);
-		ApplyRotationCorrection(PureRotA, PureRotA);
-	}
-
-	if (bLockX || bLockY || bLockZ)
-	{
-		DecomposeRotationAngles(); // TODO don't strictly need to update everything in here
-	}
-
-	if (bLockX)
-	{
-		RotateWithinLimits(0, 0, AngleX, XA, ZBProjOnX, ZA);
-	}
-
-	if (bLockY)
-	{
-		RotateWithinLimits(0, 0, AngleY, YA, ZBProjOnY, ZA);
-	}
-
-	if (bLockZ)
-	{
-		RotateWithinLimits(0, 0, AngleZ, ZA, YBProjOnZ, YA);
+		RotateToAlignAxes(ZA, ZB);
 	}
 
 	// enforce min/max angles
@@ -189,6 +161,23 @@ void FJointConstraint::UpdateJointLimits()
 	}
 }
 
+void FJointConstraint::RotateToAlignAxes(const FVector& AxisA, const FVector& AxisB) const
+{
+	const FVector ACrossB = FVector::CrossProduct(AxisA, AxisB);
+	FVector Axis;
+	float Magnitude;
+	ACrossB.ToDirectionAndLength(Axis, Magnitude);
+	const float EffectiveInvMass = FVector::DotProduct(Axis, Axis);
+	if (EffectiveInvMass < KINDA_SMALL_NUMBER)
+	{
+		return;
+	}
+	const float DeltaLambda = Magnitude / EffectiveInvMass;
+	const FVector Push = Axis * DeltaLambda;
+	const FQuat PureRotA = FQuat(Push.X, Push.Y, Push.Z, 0.0f);
+	ApplyRotationCorrection(PureRotA, PureRotA);
+}
+
 void FJointConstraint::RotateWithinLimits(
 	float MinAngle,
 	float MaxAngle,
@@ -199,15 +188,14 @@ void FJointConstraint::RotateWithinLimits(
 {
 	const bool bBeyondMin = CurrentAngle < MinAngle;
 	const bool bBeyondMax = CurrentAngle > MaxAngle;
-	if (bBeyondMin || bBeyondMax)
+	if (!(bBeyondMin || bBeyondMax))
 	{
-		const float TgtAngle = bBeyondMin ? MinAngle : MaxAngle;
-		const FQuat TgtRot = FQuat(RotAxis, FMath::DegreesToRadians(TgtAngle));
-		const FVector TgtVec = TgtRot * RefVec;
-		const FVector TgtCross = FVector::CrossProduct(TgtVec, CurVec);
-		const FQuat PureRot = FQuat(TgtCross.X, TgtCross.Y, TgtCross.Z, 0.0f);
-		ApplyRotationCorrection(PureRot, PureRot);
+		return;
 	}
+	const float TgtAngle = bBeyondMin ? MinAngle : MaxAngle;
+	const FQuat TgtRot = FQuat(RotAxis, FMath::DegreesToRadians(TgtAngle));
+	const FVector TgtVec = TgtRot * RefVec;
+	RotateToAlignAxes(TgtVec, CurVec);
 }
 
 void FJointConstraint::UpdateLocalRotateAxes(bool bX, bool bY, bool bZ)
diff --git a/Engine/Plugins/Experimental/FullBodyIK/Source/PBIK/Private/Core/PBIKSolver.cpp b/Engine/Plugins/Experimental/FullBodyIK/Source/PBIK/Private/Core/PBIKSolver.cpp
--- a/Engine/Plugins/Experimental/FullBodyIK/Source/PBIK/Private/Core/PBIKSolver.cpp
+++ b/Engine/Plugins/Experimental/FullBodyIK/Source/PBIK/Private/Core/PBIKSolver.cpp
@@ -47,7 +47,7 @@ namespace PBIK
 		}
 	}
 
-	void FEffector::SquashSubRoots()
+	void FEffector::ApplyPreferredAngles()
 	{
 		// optionally apply a preferred angle to give solver a hint which direction to favor
 		// apply amount of preferred angle proportional to the amount this sub-limb is squashed
@@ -76,12 +76,13 @@ namespace PBIK
 		}
 
 		// shrink distance to reach full blend to preferred angle
-		const float ScaledDistOrig = LengthOfChainInInputPose * 0.3f;
+		const float ScaledDistOrig = LengthOfChainInInputPose;// * 0.3f;
 		// amount squashed (clamped to scaled original length)
 		float DeltaSquash = LengthOfChainInInputPose - DistToParentSubRoot;
 		DeltaSquash = DeltaSquash > ScaledDistOrig ? ScaledDistOrig : DeltaSquash;
-		const float SquashPercent = DeltaSquash / ScaledDistOrig;
-		if (SquashPercent < 0.01f)
+		float SquashPercent = DeltaSquash / ScaledDistOrig;
+		PBIK::QuarticEaseOut(SquashPercent);
+		if (SquashPercent < 0.0001f)
 		{
 			return; // limb not squashed enough
 		}
@@ -158,19 +159,27 @@ void FPBIKSolver::Solve(const FPBIKSolverSettings& Settings)
 	// squash sub-roots to apply preferred angles
 	for (FEffector& Effector : Effectors)
 	{
-		Effector.SquashSubRoots();
+		Effector.ApplyPreferredAngles();
 	}
 
 	// run constraint iterations while allowing stretch, just to get reaching pose
 	for (int32  I = 0; I < Settings.Iterations; ++I)
 	{
+		// gradually ramp down mass of bodies
+		float IterPercent = Settings.Iterations == 1 ? 0 : (static_cast<float>(I) / static_cast<float>(Settings.Iterations - 1));
+		PBIK::SquaredEaseOut(IterPercent);
+		for (FRigidBody& Body : Bodies)
+		{
+			Body.InvMass = FMath::Lerp(Body.MaxInvMass, Body.MinInvMass, IterPercent);
+		}
+		
 		const bool bMoveSubRoots = true;
 		for (auto Constraint : Constraints)
 		{
 			Constraint->Solve(bMoveSubRoots);
 		}
 	}
-
+	
 	if (!Settings.bAllowStretch)
 	{
 		for (int32 C = Constraints.Num() - 1; C >= 0; --C)
@@ -180,11 +189,19 @@ void FPBIKSolver::Solve(const FPBIKSolverSettings& Settings)
 
 		for (FEffector& Effector : Effectors)
 		{
-			Effector.SquashSubRoots(); // update squashing once again
+			Effector.ApplyPreferredAngles(); // update squashing once again
 		}
 
 		for (int32  I = 0; I < Settings.Iterations; ++I)
 		{
+			// gradually ramp down mass of bodies
+			float IterPercent = Settings.Iterations == 1 ? 0 : (static_cast<float>(I) / static_cast<float>(Settings.Iterations - 1));
+			PBIK::SquaredEaseOut(IterPercent);
+			for (FRigidBody& Body : Bodies)
+			{
+				Body.InvMass = FMath::Lerp(Body.MaxInvMass, Body.MinInvMass, IterPercent);
+			}
+			
 			const bool bMoveSubRoots = false;
 			for (auto Constraint : Constraints)
 			{
diff --git a/Engine/Plugins/Experimental/FullBodyIK/Source/PBIK/Public/Core/PBIKBody.h b/Engine/Plugins/Experimental/FullBodyIK/Source/PBIK/Public/Core/PBIKBody.h
--- a/Engine/Plugins/Experimental/FullBodyIK/Source/PBIK/Public/Core/PBIKBody.h
+++ b/Engine/Plugins/Experimental/FullBodyIK/Source/PBIK/Public/Core/PBIKBody.h
@@ -85,6 +85,8 @@ struct FRigidBody
 	TArray<FVector> ChildLocalPositions;
 
 	float InvMass = 0.0f;
+	float MaxInvMass = 0.0f;
+	float MinInvMass = 0.0f;
 	float Length;
 	
 private:
diff --git a/Engine/Plugins/Experimental/FullBodyIK/Source/PBIK/Public/Core/PBIKConstraint.h b/Engine/Plugins/Experimental/FullBodyIK/Source/PBIK/Public/Core/PBIKConstraint.h
--- a/Engine/Plugins/Experimental/FullBodyIK/Source/PBIK/Public/Core/PBIKConstraint.h
+++ b/Engine/Plugins/Experimental/FullBodyIK/Source/PBIK/Public/Core/PBIKConstraint.h
@@ -71,6 +71,8 @@ private:
 		FVector CurVec,
 		FVector RefVec) const;
 
+	void RotateToAlignAxes(const FVector& AxisA, const FVector& AxisB) const;
+
 	void UpdateLocalRotateAxes(bool bX, bool bY, bool bZ);
 
 	void DecomposeRotationAngles();
diff --git a/Engine/Plugins/Experimental/FullBodyIK/Source/PBIK/Public/Core/PBIKSolver.h b/Engine/Plugins/Experimental/FullBodyIK/Source/PBIK/Public/Core/PBIKSolver.h
--- a/Engine/Plugins/Experimental/FullBodyIK/Source/PBIK/Public/Core/PBIKSolver.h
+++ b/Engine/Plugins/Experimental/FullBodyIK/Source/PBIK/Public/Core/PBIKSolver.h
@@ -15,6 +15,11 @@ namespace PBIK
 
 static float GLOBAL_UNITS = 100.0f; // (1.0f = meters), (100.0f = centimeters)
 
+// A long tail ease out function. Input range, 0-1. 
+FORCEINLINE static void QuarticEaseOut(float& InOut){ InOut = (FMath::Pow(InOut-1.0f, 4.0f) * -1.0f) + 1.0f; };
+// An ease out function. Input range, 0-1.
+FORCEINLINE static void SquaredEaseOut(float& InOut){ InOut = (FMath::Pow(InOut-1.0f, 2.0f) * -1.0f) + 1.0f; };
+
 struct FRigidBody;
 
 struct FEffector
@@ -46,7 +51,7 @@ struct FEffector
 
 	void UpdateFromInputs();
 
-	void SquashSubRoots();
+	void ApplyPreferredAngles();
 };
 
 } // namespace
@@ -64,6 +69,10 @@ struct PBIK_API FPBIKSolverSettings
 	UPROPERTY(EditAnywhere, Category = SolverSettings, meta = (ClampMin = "0", UIMin = "0.0", UIMax = "10.0"))
 	float MassMultiplier = 1.0f;
 
+	/** Set this as low as possible while keeping the solve stable. Lower values improve convergence of effector targets. Default is 0.2. */
+	UPROPERTY(EditAnywhere, Category = SolverSettings, meta = (ClampMin = "0", UIMin = "0.0", UIMax = "10.0"))
+	float MinMassMultiplier = 0.2f;
+
 	/** If true, joints will translate to reach the effectors; causing bones to lengthen if necessary. Good for cartoon effects. Default is false. */
 	UPROPERTY(EditAnywhere, Category = SolverSettings)
 	bool bAllowStretch = false;
