Full-body IK, minor bug fix: effectors with zero strength still pulling towards goal.

diff --git a/Engine/Plugins/Experimental/FullBodyIK/Source/PBIK/Private/Core/PBIKBody.cpp b/Engine/Plugins/Experimental/FullBodyIK/Source/PBIK/Private/Core/PBIKBody.cpp
--- a/Engine/Plugins/Experimental/FullBodyIK/Source/PBIK/Private/Core/PBIKBody.cpp
+++ b/Engine/Plugins/Experimental/FullBodyIK/Source/PBIK/Private/Core/PBIKBody.cpp
@@ -118,11 +118,6 @@ void FRigidBody::ApplyPushToRotateBody(const FVector& Push, const FVector& Offse
 
 void FRigidBody::ApplyPushToPosition(const FVector& Push)
 {
-	if (Pin && Pin->bEnabled)
-	{
-		return; // pins are locked
-	}
-
 	Position += Push * (1.0f - J.PositionStiffness);
 }
 
