Merging latest Full-Body IK fixes from UE5 Main

diff --git a/Engine/Plugins/Experimental/FullBodyIK/Source/PBIK/Private/Core/PBIKBody.cpp b/Engine/Plugins/Experimental/FullBodyIK/Source/PBIK/Private/Core/PBIKBody.cpp
--- a/Engine/Plugins/Experimental/FullBodyIK/Source/PBIK/Private/Core/PBIKBody.cpp
+++ b/Engine/Plugins/Experimental/FullBodyIK/Source/PBIK/Private/Core/PBIKBody.cpp
@@ -3,8 +3,6 @@
 #include "Core/PBIKBody.h"
 #include "Core/PBIKSolver.h"
 
-//#pragma optimize("", off)
-
 namespace PBIK
 {
 
diff --git a/Engine/Plugins/Experimental/FullBodyIK/Source/PBIK/Private/Core/PBIKSolver.cpp b/Engine/Plugins/Experimental/FullBodyIK/Source/PBIK/Private/Core/PBIKSolver.cpp
--- a/Engine/Plugins/Experimental/FullBodyIK/Source/PBIK/Private/Core/PBIKSolver.cpp
+++ b/Engine/Plugins/Experimental/FullBodyIK/Source/PBIK/Private/Core/PBIKSolver.cpp
@@ -4,8 +4,7 @@
 #include "Core/PBIKBody.h"
 #include "Core/PBIKConstraint.h"
 #include "Core/PBIKDebug.h"
-
-//#pragma optimize("", off)
+#include "PBIK.h"
 
 namespace PBIK
 {
@@ -283,13 +282,13 @@ bool FPBIKSolver::InitBones()
 
 	if (Bones.IsEmpty())
 	{
-		UE_LOG(LogTemp, Warning, TEXT("PBIK: no bones added to solver. Cannot initialize."));
+		UE_LOG(LogPBIKSolver, Warning, TEXT("PBIK: no bones added to solver. Cannot initialize."));
 		return false;
 	}
 
 	if (Effectors.IsEmpty())
 	{
-		UE_LOG(LogTemp, Warning, TEXT("PBIK: no effectors added to solver. Cannot initialize."));
+		UE_LOG(LogPBIKSolver, Warning, TEXT("PBIK: no effectors added to solver. Cannot initialize."));
 		return false;
 	}
 
@@ -306,13 +305,13 @@ bool FPBIKSolver::InitBones()
 
 	if (!SolverRoot)
 	{
-		UE_LOG(LogTemp, Warning, TEXT("PBIK: root bone not set or not found. Cannot initialize."));
+		UE_LOG(LogPBIKSolver, Warning, TEXT("PBIK: root bone not set or not found. Cannot initialize."));
 		return false;
 	}
 
 	if (NumSolverRoots > 1)
 	{
-		UE_LOG(LogTemp, Warning, TEXT("PBIK: more than 1 bone was marked as solver root. Cannot initialize."));
+		UE_LOG(LogPBIKSolver, Warning, TEXT("PBIK: more than 1 bone was marked as solver root. Cannot initialize."));
 		return false;
 	}
 
@@ -329,7 +328,7 @@ bool FPBIKSolver::InitBones()
 		const bool bIndexInRange = Bone.ParentIndex >= 0 && Bone.ParentIndex < Bones.Num();
 		if (!bIndexInRange)
 		{
-			UE_LOG(LogTemp, Warning, TEXT("PBIK: bone found with invalid parent index. Cannot initialize."));
+			UE_LOG(LogPBIKSolver, Warning, TEXT("PBIK: bone found with invalid parent index. Cannot initialize."));
 			return false;
 		}
 
@@ -391,7 +390,7 @@ bool FPBIKSolver::InitBodies()
 			FBone* BodyBone = NextBone->bIsSolverRoot ? NextBone : NextBone->Parent;
 			if (!BodyBone)
 			{
-				UE_LOG(LogTemp, Warning, TEXT("PBIK: effector is on bone that is not on or below root bone."));
+				UE_LOG(LogPBIKSolver, Warning, TEXT("PBIK: effector is on bone that is not on or below root bone."));
 				return false;
 			}
 
@@ -474,7 +473,7 @@ bool FPBIKSolver::InitConstraints()
 		FBone* BodyBone = Effector.Bone->bIsSolverRoot ? Effector.Bone : Effector.Bone->Parent;
 		if (!BodyBone)
 		{
-			UE_LOG(LogTemp, Warning, TEXT("PBIK: effector is on bone that does not have a parent."));
+			UE_LOG(LogPBIKSolver, Warning, TEXT("PBIK: effector is on bone that does not have a parent."));
 			return false;
 		}
 
@@ -582,7 +581,7 @@ PBIK::FBoneSettings* FPBIKSolver::GetBoneSettings(const int32 Index)
 
 	if (!Bones[Index].Body)
 	{
-		UE_LOG(LogTemp, Warning, TEXT("PBIK: trying to apply Bone Settings to bone that is not simulated (not between root and effector)."));
+		UE_LOG(LogPBIKSolver, Warning, TEXT("PBIK: trying to apply Bone Settings to bone that is not simulated (not between root and effector)."));
 		return nullptr;
 	}
 
@@ -599,15 +598,7 @@ void FPBIKSolver::GetBoneGlobalTransform(const int32 Index, FTransform& OutTrans
 
 int32 FPBIKSolver::GetBoneIndex(FName BoneName) const
 {
-	for (int32 B=0; B < Bones.Num(); ++B)
-	{
-		if (Bones[B].Name == BoneName)
-		{
-			return B;
-		}
-	}
-
-	return INDEX_NONE;
+	return Bones.IndexOfByPredicate([&BoneName](const PBIK::FBone& Bone) { return Bone.Name == BoneName; });
 }
 
 void FPBIKSolver::SetEffectorGoal(
diff --git a/Engine/Plugins/Experimental/FullBodyIK/Source/PBIK/Private/RigUnit_PBIK.cpp b/Engine/Plugins/Experimental/FullBodyIK/Source/PBIK/Private/RigUnit_PBIK.cpp
--- a/Engine/Plugins/Experimental/FullBodyIK/Source/PBIK/Private/RigUnit_PBIK.cpp
+++ b/Engine/Plugins/Experimental/FullBodyIK/Source/PBIK/Private/RigUnit_PBIK.cpp
@@ -3,13 +3,12 @@
 #include "RigUnit_PBIK.h"
 #include "Units/RigUnitContext.h"
 
-//#pragma optimize("", off)
 
 FRigUnit_PBIK_Execute()
 {
 	DECLARE_SCOPE_HIERARCHICAL_COUNTER_RIGUNIT()
 
-	URigHierarchy* Hierarchy = ExecuteContext.Hierarchy;
+	FRigHierarchyContainer* Hierarchy = ExecuteContext.Hierarchy;
 	if (Hierarchy == nullptr)
 	{
 		return;
@@ -39,37 +38,16 @@ FRigUnit_PBIK_Execute()
 
 		// reset all internal data
 		Solver.Reset();
-		
-		// create solver bone index to elements index map
-		TArray<FRigBoneElement*> BoneElements = Hierarchy->GetBones(true);
-		SolverBoneToElementIndex.Reset();
-		for (int B = 0; B < BoneElements.Num(); ++B)
-		{
-			SolverBoneToElementIndex.Add(BoneElements[B]->GetIndex());
-		}
 
 		// create bones
-		for (int B = 0; B < BoneElements.Num(); ++B)
+		for (int B = 0; B < Hierarchy->BoneHierarchy.Num(); ++B)
 		{
-			FName Name = BoneElements[B]->GetName();
-
-			// get the parent bone solver index
-			const int ParentElementIndex = Hierarchy->GetFirstParent(BoneElements[B]->GetIndex());
-			int ParentIndex = -1;
-			for (int P=0; P<SolverBoneToElementIndex.Num(); ++P)
-			{
-				if (SolverBoneToElementIndex[P] == ParentElementIndex)
-				{
-					ParentIndex = P;
-					break;
-				}
-			}
-			
-			const FTransform OrigTransform = Hierarchy->GetTransform(BoneElements[B], ERigTransformType::InitialGlobal);
+			const FRigBone& Bone = Hierarchy->BoneHierarchy[B];
+			const FTransform OrigTransform = Bone.InitialTransform;
 			const FVector InOrigPosition = OrigTransform.GetLocation();
 			const FQuat InOrigRotation = OrigTransform.GetRotation();
-			bool bIsRoot = BoneElements[B]->GetName() == Root;
-			Solver.AddBone(Name, ParentIndex, InOrigPosition, InOrigRotation, bIsRoot);
+			bool bIsRoot = Bone.Name == Root;
+			Solver.AddBone(Bone.Name, Bone.ParentIndex, InOrigPosition, InOrigRotation, bIsRoot);
 		}
 		
 		// create effectors
@@ -94,7 +72,7 @@ FRigUnit_PBIK_Execute()
 	// set bones to input pose
 	for(int32 BoneIndex = 0; BoneIndex < Solver.GetNumBones(); BoneIndex++)
 	{
-		const FTransform GlobalTransform = Hierarchy->GetGlobalTransform(SolverBoneToElementIndex[BoneIndex]);
+		const FTransform GlobalTransform = Hierarchy->BoneHierarchy[BoneIndex].GlobalTransform;
 		Solver.SetBoneTransform(BoneIndex, GlobalTransform);
 	}
 
@@ -152,7 +130,7 @@ FRigUnit_PBIK_Execute()
 	{
 		FTransform NewTransform;
 		Solver.GetBoneGlobalTransform(BoneIndex, NewTransform);
-		Hierarchy->SetGlobalTransform(SolverBoneToElementIndex[BoneIndex], NewTransform, false, bPropagateTransform);
+		Hierarchy->BoneHierarchy.SetGlobalTransform(BoneIndex, NewTransform, bPropagateTransform);
 	}
 
 	// do all debug drawing
diff --git a/Engine/Plugins/Experimental/FullBodyIK/Source/PBIK/Public/RigUnit_PBIK.h b/Engine/Plugins/Experimental/FullBodyIK/Source/PBIK/Public/RigUnit_PBIK.h
--- a/Engine/Plugins/Experimental/FullBodyIK/Source/PBIK/Public/RigUnit_PBIK.h
+++ b/Engine/Plugins/Experimental/FullBodyIK/Source/PBIK/Public/RigUnit_PBIK.h
@@ -98,9 +98,6 @@ struct FRigUnit_PBIK : public FRigUnit_HighlevelBaseMutable
 	UPROPERTY(transient)
 	TArray<int32> BoneSettingToSolverBoneIndex;
 
-	UPROPERTY(transient)
-	TArray<int32> SolverBoneToElementIndex;
-
 	UPROPERTY(transient)
 	FPBIKSolver Solver;
 };
