Full-body IK fixes:
1. Can now lock all 3 axes
2. Better behavior when pulling effectors away from body.
3. Added PullStrengthAlpha param to effectors

diff --git a/Engine/Plugins/Animation/IKRig/Source/IKRig/Private/Solvers/IKRig_PBIKSolver.cpp b/Engine/Plugins/Animation/IKRig/Source/IKRig/Private/Solvers/IKRig_PBIKSolver.cpp
--- a/Engine/Plugins/Animation/IKRig/Source/IKRig/Private/Solvers/IKRig_PBIKSolver.cpp
+++ b/Engine/Plugins/Animation/IKRig/Source/IKRig/Private/Solvers/IKRig_PBIKSolver.cpp
@@ -96,7 +96,9 @@ void UIKRigPBIKSolver::Solve(FIKRigSkeleton& IKRigSkeleton, const FIKRigGoalCont
 			Goal.FinalBlendedPosition,
 			Goal.FinalBlendedRotation,
 			OffsetAlpha,
-			Effector->StrengthAlpha);
+			Effector->StrengthAlpha,
+			Effector->PullChainAlpha,
+			Effector->bPinRotation);
 	}
 
 	// update settings
@@ -107,6 +109,7 @@ void UIKRigPBIKSolver::Solve(FIKRigSkeleton& IKRigSkeleton, const FIKRigGoalCont
 	Settings.MinMassMultiplier = MinMassMultiplier;
 	Settings.bPinRoot = bPinRoot;
 	Settings.bStartSolveFromInputPose = bStartSolveFromInputPose;
+	Settings.bPreProcessPose = bPreProcessPose;
 
 	// solve
 	Solver.Solve(Settings);
@@ -130,6 +133,7 @@ void UIKRigPBIKSolver::UpdateSolverSettings(UIKRigSolver* InSettings)
 		MinMassMultiplier = Settings->MinMassMultiplier;
 		bPinRoot = Settings->bPinRoot;
 		bStartSolveFromInputPose = Settings->bStartSolveFromInputPose;
+		bPreProcessPose = Settings->bPreProcessPose;
 
 		// copy effector settings
 		for (const UIKRig_FBIKEffector* InEffector : Settings->Effectors)
diff --git a/Engine/Plugins/Animation/IKRig/Source/IKRig/Public/Solvers/IKRig_PBIKSolver.h b/Engine/Plugins/Animation/IKRig/Source/IKRig/Public/Solvers/IKRig_PBIKSolver.h
--- a/Engine/Plugins/Animation/IKRig/Source/IKRig/Public/Solvers/IKRig_PBIKSolver.h
+++ b/Engine/Plugins/Animation/IKRig/Source/IKRig/Public/Solvers/IKRig_PBIKSolver.h
@@ -24,6 +24,12 @@ public:
 	
 	UPROPERTY(EditAnywhere, Category = "Full Body IK Effector", meta = (ClampMin = "0", ClampMax = "1", UIMin = "0.0", UIMax = "1.0"))
 	float StrengthAlpha = 1.0f;
+	
+	UPROPERTY(EditAnywhere, Category = "Full Body IK Effector", meta = (ClampMin = "0", ClampMax = "1", UIMin = "0.0", UIMax = "1.0"))
+	float PullChainAlpha = 1.0f;
+
+	UPROPERTY(EditAnywhere, Category = "Full Body IK Effector")
+	bool bPinRotation = true;
 
 	UPROPERTY(Transient)
 	int32 IndexInSolver = -1;
@@ -31,6 +37,8 @@ public:
 	void CopySettings(const UIKRig_FBIKEffector* Other)
 	{
 		StrengthAlpha = Other->StrengthAlpha;
+		PullChainAlpha = Other->PullChainAlpha;
+		bPinRotation = Other->bPinRotation;
 	}
 };
 
@@ -130,7 +138,7 @@ public:
 	FName RootBone;
 
 	/** High iteration counts can help solve complex joint configurations with competing constraints, but will increase runtime cost. Default is 20. */
-	UPROPERTY(EditAnywhere, Category = SolverSettings, meta = (ClampMin = "0", UIMin = "0.0", UIMax = "200.0"))
+	UPROPERTY(EditAnywhere, Category = SolverSettings, meta = (ClampMin = "0", UIMin = "0.0", UIMax = "100.0"))
 	int32 Iterations = 20;
 
 	/** A global mass multiplier; higher values will make the joints more stiff, but require more iterations. Typical range is 0.0 to 10.0. */
@@ -145,6 +153,10 @@ public:
 	UPROPERTY(EditAnywhere, Category = SolverSettings)
 	bool bAllowStretch = false;
 
+	/** If true, solver will pre-process the skeleton to approximate a solved pose before running constraint iterations. Default is true. */
+	UPROPERTY(EditAnywhere, Category = SolverSettings)
+	bool bPreProcessPose = true;
+
 	/** Lock the position and rotation of the solver root bone in-place (at animated position). Useful for partial-body solves. Default is false. */
 	UPROPERTY(EditAnywhere, Category = SolverSettings)
 	bool bPinRoot = false;
diff --git a/Engine/Plugins/Experimental/FullBodyIK/Source/PBIK/Private/Core/PBIKBody.cpp b/Engine/Plugins/Experimental/FullBodyIK/Source/PBIK/Private/Core/PBIKBody.cpp
--- a/Engine/Plugins/Experimental/FullBodyIK/Source/PBIK/Private/Core/PBIKBody.cpp
+++ b/Engine/Plugins/Experimental/FullBodyIK/Source/PBIK/Private/Core/PBIKBody.cpp
@@ -3,8 +3,6 @@
 #include "Core/PBIKBody.h"
 #include "Core/PBIKSolver.h"
 
-//#pragma optimize("", off)
-
 namespace PBIK
 {
 
@@ -20,7 +18,6 @@ FBone::FBone(
 	Position = InOrigPosition;
 	Rotation = InOrigRotation;
 	bIsSolverRoot = bInIsSolverRoot;
-	bIsSolved = false; // default
 }
 
 bool FBone::HasChild(const FBone* Bone)
@@ -36,6 +33,18 @@ bool FBone::HasChild(const FBone* Bone)
 	return false;
 }
 
+void FBone::UpdateFromInputs()
+{
+	if (!Parent)
+	{
+		return;
+	}
+
+	LocalPositionOrig = Parent->Rotation.Inverse() * (Position - Parent->Position);
+	LocalRotationOrig = Parent->Rotation.Inverse() * Rotation;
+	Length = LocalPositionOrig.Size();
+}
+
 FRigidBody::FRigidBody(FBone* InBone)
 {
 	Bone = InBone;
@@ -45,16 +54,16 @@ FRigidBody::FRigidBody(FBone* InBone)
 void FRigidBody::Initialize(FBone* SolverRoot)
 {
 	FVector Centroid = Bone->Position;
-	Length = 0.0f;
+	Mass = 0.0f;
 	for(const FBone* Child : Bone->Children)
 	{
 		Centroid += Child->Position;
-		Length += (Bone->Position - Child->Position).Size();
+		Mass += (Bone->Position - Child->Position).Size();
 	}
-	Centroid = Centroid * (1.0f / ((float)Bone->Children.Num() + 1.0f));
+	Centroid = Centroid * (1.0f / (Bone->Children.Num() + 1.0f));
 
 	Position = Centroid;
-	Rotation = RotationOrig = Bone->Rotation;
+	Rotation = InitialRotation = Bone->Rotation;
 	BoneLocalPosition = Rotation.Inverse() * (Bone->Position - Centroid);
 
 	for (FBone* Child : Bone->Children)
@@ -80,12 +89,13 @@ void FRigidBody::UpdateFromInputs(const FPBIKSolverSettings& Settings)
 		// set to input pose
 		Position = Bone->Position - Bone->Rotation * BoneLocalPosition;
 		Rotation = Bone->Rotation;
+		InputPosition = Position;
 	}
 
-	// Body.Length used as rough approximation of the mass of the body
 	// for fork joints (multiple solved children) we sum lengths to all children (see Initialize)
-	MaxInvMass = 1.0f / ( Length * ((Settings.MassMultiplier * GLOBAL_UNITS) + 0.5f));
-	MinInvMass = 1.0f / ( Length * ((Settings.MinMassMultiplier * GLOBAL_UNITS) + 0.5f));
+	const float MinMass = 0.5f; // prevent mass ever hitting zero
+	MaxInvMass = 1.0f / (Mass * ((Settings.MassMultiplier * GLOBAL_UNITS) + MinMass));
+	MinInvMass = 1.0f / (Mass * ((Settings.MinMassMultiplier * GLOBAL_UNITS) + MinMass));
 }
 
 int FRigidBody::GetNumBonesToRoot() const
@@ -111,9 +121,9 @@ void FRigidBody::ApplyPushToRotateBody(const FVector& Push, const FVector& Offse
 	}
 	
 	// equation 8 in "Detailed Rigid Body Simulation with XPBD"
-	FVector Omega = InvMass * (1.0f - J.RotationStiffness) * FVector::CrossProduct(Offset, Push);
-	FQuat OQ(Omega.X, Omega.Y, Omega.Z, 0.0f);
-	ApplyRotationDelta(OQ, false);
+	const FVector Omega = InvMass * (1.0f - J.RotationStiffness) * FVector::CrossProduct(Offset, Push);
+	const FQuat DeltaQ(Omega.X, Omega.Y, Omega.Z, 0.0f);
+	ApplyRotationDelta(DeltaQ);
 }
 
 void FRigidBody::ApplyPushToPosition(const FVector& Push)
@@ -121,23 +131,23 @@ void FRigidBody::ApplyPushToPosition(const FVector& Push)
 	Position += Push * (1.0f - J.PositionStiffness);
 }
 
-void FRigidBody::ApplyRotationDelta(const FQuat& InDelta, const bool bNegated)
+void FRigidBody::ApplyRotationDelta(const FQuat& DeltaQ)
 {
 	if (Pin && Pin->bEnabled && Pin->bPinRotation)
 	{
 		return; // rotation of this body is pinned
 	}
 
-	/** InDelta is assumed to be a "pure" quaternion representing an infintesimal rotation */
-	FQuat Delta = InDelta * Rotation;
+	/** DeltaQ is assumed to be a "pure" quaternion representing an infintesimal rotation */
+	FQuat Delta = DeltaQ * Rotation;
 	Delta.X *= 0.5f;
 	Delta.Y *= 0.5f;
 	Delta.Z *= 0.5f;
 	Delta.W *= 0.5f;
-	Rotation.X = Rotation.X + (bNegated ? -Delta.X : Delta.X);
-	Rotation.Y = Rotation.Y + (bNegated ? -Delta.Y : Delta.Y);
-	Rotation.Z = Rotation.Z + (bNegated ? -Delta.Z : Delta.Z);
-	Rotation.W = Rotation.W + (bNegated ? -Delta.W : Delta.W);
+	Rotation.X = Rotation.X + Delta.X;
+	Rotation.Y = Rotation.Y + Delta.Y;
+	Rotation.Z = Rotation.Z + Delta.Z;
+	Rotation.W = Rotation.W + Delta.W;
 	Rotation.Normalize();
 }
 } // namespace
\ No newline at end of file
diff --git a/Engine/Plugins/Experimental/FullBodyIK/Source/PBIK/Private/Core/PBIKConstraint.cpp b/Engine/Plugins/Experimental/FullBodyIK/Source/PBIK/Private/Core/PBIKConstraint.cpp
--- a/Engine/Plugins/Experimental/FullBodyIK/Source/PBIK/Private/Core/PBIKConstraint.cpp
+++ b/Engine/Plugins/Experimental/FullBodyIK/Source/PBIK/Private/Core/PBIKConstraint.cpp
@@ -16,9 +16,9 @@ FJointConstraint::FJointConstraint(FRigidBody* InA, FRigidBody* InB)
 	PinPointLocalToA = A->Rotation.Inverse() * (PinPoint - A->Position);
 	PinPointLocalToB = B->Rotation.Inverse() * (PinPoint - B->Position);
 
-	XOrig = B->RotationOrig * FVector(1, 0, 0);
-	YOrig = B->RotationOrig * FVector(0, 1, 0);
-	ZOrig = B->RotationOrig * FVector(0, 0, 1);
+	XOrig = B->InitialRotation * FVector(1, 0, 0);
+	YOrig = B->InitialRotation * FVector(0, 1, 0);
+	ZOrig = B->InitialRotation * FVector(0, 0, 1);
 
 	AngleX = AngleY = AngleZ = 0;
 }
@@ -86,10 +86,10 @@ FVector FJointConstraint::GetPositionCorrection(FVector& OutBodyToA, FVector& Ou
 	return (PinPointOnB - PinPointOnA);
 }
 
-void FJointConstraint::ApplyRotationCorrection(FQuat PureRotA, FQuat PureRotB) const
+void FJointConstraint::ApplyRotationCorrection(FQuat DeltaQ) const
 {
-	A->ApplyRotationDelta(PureRotA, false);
-	B->ApplyRotationDelta(PureRotB, true);
+	A->ApplyRotationDelta(DeltaQ);
+	B->ApplyRotationDelta(DeltaQ * -1.0f);
 }
 
 void FJointConstraint::UpdateJointLimits()
@@ -114,7 +114,18 @@ void FJointConstraint::UpdateJointLimits()
 	const bool bLockY = (J.Y == ELimitType::Locked) || (J.Y == ELimitType::Limited && ((J.MaxY - J.MinY) < 2.0f));
 	const bool bLockZ = (J.Z == ELimitType::Locked) || (J.Z == ELimitType::Limited && ((J.MaxZ - J.MinZ) < 2.0f));
 
-	// determine which axes should be treated as a hinge (mutually exclusive)
+	// when locked on THREE axes, consider the bone FIXED
+	if (bLockX && bLockY && bLockZ)
+	{
+		const FQuat OrigOffset = A->InitialRotation * B->InitialRotation.Inverse();
+		const FQuat CurrentOffset = A->Rotation * B->Rotation.Inverse();
+		const FQuat DeltaOffset = OrigOffset * CurrentOffset.Inverse();
+		const FQuat DeltaQ = FQuat(DeltaOffset.X, DeltaOffset.Y, DeltaOffset.Z, 0.0f) * 2.0f;
+		ApplyRotationCorrection(DeltaQ);
+		return; // fixed bones cannot move at all
+	}
+
+	// when locked on TWO axes, consider the bone a HINGE (mutually exclusive)
 	const bool bXHinge = !bLockX && bLockY && bLockZ;
 	const bool bYHinge = bLockX && !bLockY && bLockZ;
 	const bool bZHinge = bLockX && bLockY && !bLockZ;
@@ -167,15 +178,15 @@ void FJointConstraint::RotateToAlignAxes(const FVector& AxisA, const FVector& Ax
 	FVector Axis;
 	float Magnitude;
 	ACrossB.ToDirectionAndLength(Axis, Magnitude);
-	const float EffectiveInvMass = FVector::DotProduct(Axis, Axis);
+	const float EffectiveInvMass = FVector::DotProduct(Axis, Axis); // Todo incorporate Body Mass here?
 	if (EffectiveInvMass < KINDA_SMALL_NUMBER)
 	{
 		return;
 	}
 	const float DeltaLambda = Magnitude / EffectiveInvMass;
 	const FVector Push = Axis * DeltaLambda;
-	const FQuat PureRotA = FQuat(Push.X, Push.Y, Push.Z, 0.0f);
-	ApplyRotationCorrection(PureRotA, PureRotA);
+	const FQuat DeltaQ = FQuat(Push.X, Push.Y, Push.Z, 0.0f);
+	ApplyRotationCorrection(DeltaQ);
 }
 
 void FJointConstraint::RotateWithinLimits(
@@ -200,8 +211,8 @@ void FJointConstraint::RotateWithinLimits(
 
 void FJointConstraint::UpdateLocalRotateAxes(bool bX, bool bY, bool bZ)
 {
-	const FQuat ARot = A->Rotation * A->RotationOrig.Inverse();
-	const FQuat BRot = B->Rotation * B->RotationOrig.Inverse();
+	const FQuat ARot = A->Rotation * A->InitialRotation.Inverse();
+	const FQuat BRot = B->Rotation * B->InitialRotation.Inverse();
 
 	if (bX)
 	{
@@ -224,8 +235,8 @@ void FJointConstraint::UpdateLocalRotateAxes(bool bX, bool bY, bool bZ)
 
 void FJointConstraint::DecomposeRotationAngles()
 {
-	const FQuat ARot = A->Rotation * A->RotationOrig.Inverse();
-	const FQuat BRot = B->Rotation * B->RotationOrig.Inverse();
+	const FQuat ARot = A->Rotation * A->InitialRotation.Inverse();
+	const FQuat BRot = B->Rotation * B->InitialRotation.Inverse();
 
 	XA = ARot * XOrig;
 	YA = ARot * YOrig;
@@ -250,7 +261,7 @@ float FJointConstraint::SignedAngleBetweenNormals(
 {
 	const float FromDotTo = FVector::DotProduct(From, To);
 	const float Angle = FMath::RadiansToDegrees(FMath::Acos(FromDotTo));
-	const FVector Cross = FVector::CrossProduct(From, To); // TODO may be backwards?
+	const FVector Cross = FVector::CrossProduct(From, To);
 	const float Dot = FVector::DotProduct(Cross, Axis);
 	return Dot >= 0 ? Angle : -Angle;
 }
diff --git a/Engine/Plugins/Experimental/FullBodyIK/Source/PBIK/Private/Core/PBIKSolver.cpp b/Engine/Plugins/Experimental/FullBodyIK/Source/PBIK/Private/Core/PBIKSolver.cpp
--- a/Engine/Plugins/Experimental/FullBodyIK/Source/PBIK/Private/Core/PBIKSolver.cpp
+++ b/Engine/Plugins/Experimental/FullBodyIK/Source/PBIK/Private/Core/PBIKSolver.cpp
@@ -4,8 +4,7 @@
 #include "Core/PBIKBody.h"
 #include "Core/PBIKConstraint.h"
 #include "Core/PBIKDebug.h"
-
-//#pragma optimize("", off)
+#include "PBIK.h"
 
 namespace PBIK
 {
@@ -13,14 +12,16 @@ namespace PBIK
 	{
 		check(InBone);
 		Bone = InBone;
-		SetGoal(Bone->Position, Bone->Rotation, 1.0f, 1.0f);
+		SetGoal(Bone->Position, Bone->Rotation, 1.0f, 1.0f, 1.0f, false);
 	}
 
 	void FEffector::SetGoal(
 		const FVector& InPositionGoal,
 		const FQuat& InRotationGoal,
 		float InTransformAlpha,
-		float InStrengthAlpha)
+		float InStrengthAlpha,
+		float InPullChainAlpha,
+		bool bInPinRotation)
 	{
 		PositionOrig = Bone->Position;
 		RotationOrig = Bone->Rotation;
@@ -30,10 +31,13 @@ namespace PBIK
 
 		TransformAlpha = InTransformAlpha;
 		StrengthAlpha = InStrengthAlpha;
+		PullChainAlpha = InPullChainAlpha;
+		bPinRotation = bInPinRotation;
 	}
 
-	void FEffector::UpdateFromInputs()
+	void FEffector::UpdateFromInputs(const FBone& SolverRoot)
 	{
+		// blend effector transform by alpha and set pin goal transform
 		Position = FMath::Lerp(PositionOrig, PositionGoal, TransformAlpha);
 		Rotation = FMath::Lerp(RotationOrig, RotationGoal, TransformAlpha);
 		Pin.Pin()->SetGoal(Position, Rotation, StrengthAlpha);
@@ -42,8 +46,40 @@ namespace PBIK
 		if (ParentSubRoot)
 		{
 			FEffector* ParentEffector = ParentSubRoot->Effector;
-			const FVector ParentSubRootPosition = ParentEffector ? ParentEffector->PositionOrig : ParentSubRoot->Bone->Position;
-			LengthOfChainInInputPose = (ParentSubRootPosition - Bone->Position).Size();
+			const FVector ParentSubRootPosition = ParentEffector ? ParentEffector->PositionOrig : ParentSubRoot->Position;
+			DistanceToSubRootInInputPose = (ParentSubRootPosition - Bone->Position).Size();
+		}
+
+		// update distances to root (along bones)
+		{
+			if (Bone->bIsSolverRoot)
+			{
+				DistToRootAlongBones = 0.0f;
+			}else
+			{
+				DistToRootAlongBones = Bone->Length;
+				FBone* Parent = Bone->Parent;
+				while (Parent && !Parent->bIsSolverRoot)
+				{
+					DistToRootAlongBones += Parent->Length;
+					Parent = Parent->Parent;
+				}
+			}
+		}
+
+		// update distances to SUB root (along bones)
+		if (ParentSubRoot)
+		{
+			DistancesFromEffector.Reset();
+			DistancesFromEffector.Add(0.0f);
+			DistToSubRootAlongBones = Bone->Length;
+			FBone* Parent = Bone->Parent;
+			while (Parent && Parent->Body != ParentSubRoot)
+			{
+				DistancesFromEffector.Add(DistToSubRootAlongBones);
+				DistToSubRootAlongBones += Parent->Length;
+				Parent = Parent->Parent;
+			}
 		}
 	}
 
@@ -59,7 +95,7 @@ namespace PBIK
 		}
 
 		// can't squash chain with zero length already
-		if (LengthOfChainInInputPose <= SMALL_NUMBER)
+		if (DistanceToSubRootInInputPose <= SMALL_NUMBER)
 		{
 			return;
 		}
@@ -70,15 +106,15 @@ namespace PBIK
 		FEffector* ParentEffector = ParentSubRoot->Effector;
 		const FVector ParentSubRootPosition = ParentEffector ? ParentEffector->Position : ParentSubRoot->Position;
 		const float DistToParentSubRoot = (ParentSubRootPosition - Position).Size();
-		if (DistToParentSubRoot >= LengthOfChainInInputPose)
+		if (DistToParentSubRoot >= DistanceToSubRootInInputPose)
 		{
 			return; // limb is stretched
 		}
 
 		// shrink distance to reach full blend to preferred angle
-		const float ScaledDistOrig = LengthOfChainInInputPose;// * 0.3f;
+		const float ScaledDistOrig = DistanceToSubRootInInputPose;// * 0.3f;
 		// amount squashed (clamped to scaled original length)
-		float DeltaSquash = LengthOfChainInInputPose - DistToParentSubRoot;
+		float DeltaSquash = DistanceToSubRootInInputPose - DistToParentSubRoot;
 		DeltaSquash = DeltaSquash > ScaledDistOrig ? ScaledDistOrig : DeltaSquash;
 		float SquashPercent = DeltaSquash / ScaledDistOrig;
 		PBIK::QuarticEaseOut(SquashPercent);
@@ -105,7 +141,6 @@ namespace PBIK
 			Parent = Parent->Parent;
 		}
 	}
-
 } // namespace
 
 void FPBIKSolver::Solve(const FPBIKSolverSettings& Settings)
@@ -128,14 +163,7 @@ void FPBIKSolver::Solve(const FPBIKSolverSettings& Settings)
 	// even the LocalPosition has to be updated in case translation is animated
 	for (FBone& Bone : Bones)
 	{
-		FBone* Parent = Bone.Parent;
-		if (!Parent)
-		{
-			continue;
-		}
-
-		Bone.LocalPositionOrig = Parent->Rotation.Inverse() * (Bone.Position - Parent->Position);
-		Bone.LocalRotationOrig = Parent->Rotation.Inverse() * Bone.Rotation;
+		Bone.UpdateFromInputs();
 	}
 
 	// update Bodies with new bone positions from incoming pose and solver settings
@@ -150,100 +178,210 @@ void FPBIKSolver::Solve(const FPBIKSolverSettings& Settings)
 		RootPin.Pin()->bEnabled = Settings.bPinRoot;
 	}
 
-	// blend effectors by Alpha and update pin goals
+	// blend effectors by Alpha, update pin goals and update effector dist to root
 	for (FEffector& Effector : Effectors)
 	{
-		Effector.UpdateFromInputs();
+		Effector.UpdateFromInputs(*SolverRoot);
 	}
-	
-	// squash sub-roots to apply preferred angles
-	for (FEffector& Effector : Effectors)
+
+	// do all constraint solving
+	UpdateBodies(Settings);
+
+	// update Bone transforms controlled by Bodies
+	for (FRigidBody& Body : Bodies)
 	{
-		Effector.ApplyPreferredAngles();
+		Body.Bone->Position = Body.Position + Body.Rotation * Body.BoneLocalPosition;
+		Body.Bone->Rotation = Body.Rotation;
 	}
 
-	// run constraint iterations while allowing stretch, just to get reaching pose
-	for (int32  I = 0; I < Settings.Iterations; ++I)
+	// update Bone transforms controlled by effectors
+	for (const FEffector& Effector : Effectors)
 	{
-		// gradually ramp down mass of bodies
-		float IterPercent = Settings.Iterations == 1 ? 0 : (static_cast<float>(I) / static_cast<float>(Settings.Iterations - 1));
-		PBIK::SquaredEaseOut(IterPercent);
-		for (FRigidBody& Body : Bodies)
+		FBone* Bone = Effector.Bone;
+		if (Bone->bIsSolverRoot)
 		{
-			Body.InvMass = FMath::Lerp(Body.MaxInvMass, Body.MinInvMass, IterPercent);
+			continue; // if there's an effector on the root, leave it where the body ended up
+		}
+
+		if (Bone->Body)
+		{
+			continue; // effector is between other effectors, so leave transform where body ended up
 		}
 		
-		const bool bMoveSubRoots = true;
-		for (auto Constraint : Constraints)
+		Bone->Position = Bone->Parent->Position + Bone->Parent->Rotation * Bone->LocalPositionOrig;
+
+		if (Effector.bPinRotation)
 		{
-			Constraint->Solve(bMoveSubRoots);
+			Bone->Rotation = Effector.Rotation; // optionally pin rotation to that of effector
+		}else
+		{
+			Bone->Rotation = Bone->Parent->Rotation * Bone->LocalRotationOrig;
+		}
+	}
+
+	// propagate to non-solved bones (requires storage in root to tip order)
+	for (FBone& Bone : Bones)
+	{
+		if (Bone.bIsSolved || !Bone.Parent)
+		{
+			continue;
 		}
+
+		Bone.Position = Bone.Parent->Position + Bone.Parent->Rotation * Bone.LocalPositionOrig;
+		Bone.Rotation = Bone.Parent->Rotation * Bone.LocalRotationOrig;
 	}
+}
+
+void FPBIKSolver::PullRootTowardsEffectors()
+{
+	using PBIK::FEffector;
+	using PBIK::FRigidBody;
 	
-	if (!Settings.bAllowStretch)
+	// accumulate a delta offset from each effector
+	FVector BodyOffset = FVector::ZeroVector;
+	for (FEffector& Effector : Effectors)
 	{
-		for (int32 C = Constraints.Num() - 1; C >= 0; --C)
+		const float DistToRootStraightLine = (Effector.Position - SolverRoot->Position).Size();
+		const float Delta = DistToRootStraightLine - Effector.DistToRootAlongBones;
+		if (Delta < 0.0f)
 		{
-			Constraints[C]->RemoveStretch();
+			continue; // only pull
 		}
 
-		for (FEffector& Effector : Effectors)
+		FVector RootToTip = Effector.Position - SolverRoot->Position;
+		BodyOffset += Delta * RootToTip.GetSafeNormal();
+	}
+
+	const float InvNumEffectors = (1.0f / static_cast<float>(Effectors.Num()));
+	BodyOffset *= InvNumEffectors;
+
+	// linearly move entire apparatus towards tip effectors
+	for (FRigidBody& Body : Bodies)
+	{
+		Body.Position += BodyOffset;
+	}
+}
+
+void FPBIKSolver::PullChainsTowardsEffectors()
+{
+	using PBIK::FEffector;
+	using PBIK::FBone;
+	using PBIK::FRigidBody;
+	
+	for (FEffector& Effector : Effectors)
+	{
+		if (!Effector.ParentSubRoot)
 		{
-			Effector.ApplyPreferredAngles(); // update squashing once again
+			continue;
 		}
 
-		for (int32  I = 0; I < Settings.Iterations; ++I)
+		if (Effector.DistToSubRootAlongBones < SMALL_NUMBER)
 		{
-			// gradually ramp down mass of bodies
-			float IterPercent = Settings.Iterations == 1 ? 0 : (static_cast<float>(I) / static_cast<float>(Settings.Iterations - 1));
-			PBIK::SquaredEaseOut(IterPercent);
-			for (FRigidBody& Body : Bodies)
-			{
-				Body.InvMass = FMath::Lerp(Body.MaxInvMass, Body.MinInvMass, IterPercent);
-			}
-			
-			const bool bMoveSubRoots = false;
-			for (auto Constraint : Constraints)
-			{
-				Constraint->Solve(bMoveSubRoots);
-			}
+			continue;
 		}
 
-		for (int32  C = Constraints.Num() - 1; C >= 0; --C)
+		// get original chain vector
+		const FVector ChainStartOrig = Effector.ParentSubRoot->InputPosition;
+		const FVector ChainEndOrig = Effector.PositionOrig;
+		FVector ChainVecOrig;
+		float ChainLenOrig;
+		(ChainEndOrig - ChainStartOrig).ToDirectionAndLength(ChainVecOrig, ChainLenOrig);
+
+		// get new chain vector
+		FVector ChainStartNew;
+		if (Effector.ParentSubRoot->Effector)
 		{
-			Constraints[C]->RemoveStretch();
+			// add effector offset to chain root
+			FEffector* RootEffector = Effector.ParentSubRoot->Effector;
+			FVector RootEffectorDelta = RootEffector->Position - RootEffector->PositionOrig;
+			ChainStartNew = Effector.ParentSubRoot->Position + RootEffectorDelta;
+		}else
+		{
+			ChainStartNew = Effector.ParentSubRoot->Position;
+		}
+		const FVector ChainEndNew = Effector.Position;
+		FVector ChainVecNew;
+		float ChainLenNew;
+		(ChainEndNew - ChainStartNew).ToDirectionAndLength(ChainVecNew, ChainLenNew);
+
+		FQuat ChainDeltaRotation = FQuat::FindBetweenNormals(ChainVecOrig, ChainVecNew);
+		ChainDeltaRotation = FQuat::FastLerp(FQuat::Identity, ChainDeltaRotation, (Effector.PullChainAlpha * Effector.StrengthAlpha));
+		FVector ChainDeltaPosition = ChainVecNew * (ChainLenNew - ChainLenOrig) * (Effector.PullChainAlpha * Effector.StrengthAlpha);
+
+		FBone* Bone = Effector.Bone->Parent;
+		int32 ChainIndex = 0;
+		const float InvChainLength = 1.0f / Effector.DistToSubRootAlongBones;
+		while (Bone && Bone->Body != Effector.ParentSubRoot)
+		{		
+			// rotate body along with chain
+			const FVector BodyRelativeToChain = Bone->Body->Position - ChainStartNew;
+			const FVector RotatedBodyPos = ChainStartNew + ChainDeltaRotation.RotateVector(BodyRelativeToChain);
+			Bone->Body->Position = RotatedBodyPos;
+			Bone->Body->Rotation = ChainDeltaRotation * Bone->Body->Rotation;
+
+			// move body proportional to chain stretch amount
+			const float Strength = 1.0f - (Effector.DistancesFromEffector[ChainIndex] * InvChainLength);
+			Bone->Body->Position += ChainDeltaPosition * Strength;
+			
+			Bone = Bone->Parent;
+			++ChainIndex;
 		}
 	}
+}
 
-	// update Bone transforms controlled by Bodies
-	for (FRigidBody& Body : Bodies)
+void FPBIKSolver::UpdateBodies(const FPBIKSolverSettings& Settings)
+{
+	using PBIK::FRigidBody;
+	using PBIK::FEffector;
+
+	// this creates a better pose from which to start the constraint solving
+	if (Settings.bPreProcessPose)
 	{
-		Body.Bone->Position = Body.Position + Body.Rotation * Body.BoneLocalPosition;
-		Body.Bone->Rotation = Body.Rotation;
+		// pull ALL bodies towards the effectors
+		PullRootTowardsEffectors();
+		// pull each effector's chain towards itself
+		PullChainsTowardsEffectors();
 	}
-
-	// update Bone transforms controlled by effectors
-	for (const FEffector& Effector : Effectors)
+	
+	// apply preferred angles to squashed effector chains
+	for (FEffector& Effector : Effectors)
 	{
-		FBone* Bone = Effector.Bone;
-		if (Bone->bIsSolverRoot)
-		{
-			continue; // if there's an effector on the root, leave it where the body ended up
-		}
-		Bone->Position = Bone->Parent->Position + Bone->Parent->Rotation * Bone->LocalPositionOrig;
-		Bone->Rotation = Effector.Rotation;
+		Effector.ApplyPreferredAngles();
 	}
+	
+	// run constraint iterations while allowing stretch, just to get reaching pose
+	SolveConstraints(Settings.Iterations, true, Settings.bAllowStretch);
+}
 
-	// propagate to non-solved bones (requires storage in root to tip order)
-	for (FBone& Bone : Bones)
+void FPBIKSolver::SolveConstraints(const int32 Iterations, const bool bMoveRoots, const bool bAllowStretch)
+{
+	using PBIK::FRigidBody;
+	
+	for (int32  I = 0; I < Iterations; ++I)
 	{
-		if (Bone.bIsSolved || !Bone.Parent)
+		float IterPercent = Iterations == 1 ? 0 : (static_cast<float>(I) / static_cast<float>(Iterations - 1));
+		
+		// gradually ramp down mass of bodies
+		float MassRamp = PBIK::SquaredEaseOut(IterPercent);
+		for (FRigidBody& Body : Bodies)
 		{
-			continue;
+			Body.InvMass = FMath::Lerp(Body.MaxInvMass, Body.MinInvMass, MassRamp);
+		}
+		
+		// solve all constraints
+		for (auto Constraint : Constraints)
+		{
+			Constraint->Solve(bMoveRoots);
 		}
 
-		Bone.Position = Bone.Parent->Position + Bone.Parent->Rotation * Bone.LocalPositionOrig;
-		Bone.Rotation = Bone.Parent->Rotation * Bone.LocalRotationOrig;
+		// do post-pass to remove stretch
+		if (!bAllowStretch && (I>Iterations/2))
+		{
+			for (int32 C = Constraints.Num() - 1; C >= 0; --C)
+			{
+				Constraints[C]->RemoveStretch();
+			}
+		}
 	}
 }
 
@@ -283,13 +421,13 @@ bool FPBIKSolver::InitBones()
 
 	if (Bones.IsEmpty())
 	{
-		UE_LOG(LogTemp, Warning, TEXT("PBIK: no bones added to solver. Cannot initialize."));
+		UE_LOG(LogPBIKSolver, Warning, TEXT("PBIK: no bones added to solver. Cannot initialize."));
 		return false;
 	}
 
 	if (Effectors.IsEmpty())
 	{
-		UE_LOG(LogTemp, Warning, TEXT("PBIK: no effectors added to solver. Cannot initialize."));
+		UE_LOG(LogPBIKSolver, Warning, TEXT("PBIK: no effectors added to solver. Cannot initialize."));
 		return false;
 	}
 
@@ -306,13 +444,13 @@ bool FPBIKSolver::InitBones()
 
 	if (!SolverRoot)
 	{
-		UE_LOG(LogTemp, Warning, TEXT("PBIK: root bone not set or not found. Cannot initialize."));
+		UE_LOG(LogPBIKSolver, Warning, TEXT("PBIK: root bone not set or not found. Cannot initialize."));
 		return false;
 	}
 
 	if (NumSolverRoots > 1)
 	{
-		UE_LOG(LogTemp, Warning, TEXT("PBIK: more than 1 bone was marked as solver root. Cannot initialize."));
+		UE_LOG(LogPBIKSolver, Warning, TEXT("PBIK: more than 1 bone was marked as solver root. Cannot initialize."));
 		return false;
 	}
 
@@ -329,7 +467,7 @@ bool FPBIKSolver::InitBones()
 		const bool bIndexInRange = Bone.ParentIndex >= 0 && Bone.ParentIndex < Bones.Num();
 		if (!bIndexInRange)
 		{
-			UE_LOG(LogTemp, Warning, TEXT("PBIK: bone found with invalid parent index. Cannot initialize."));
+			UE_LOG(LogPBIKSolver, Warning, TEXT("PBIK: bone found with invalid parent index. Cannot initialize."));
 			return false;
 		}
 
@@ -337,7 +475,7 @@ bool FPBIKSolver::InitBones()
 		Bone.Parent = &Bones[Bone.ParentIndex];
 	}
 
-	// initialize IsSolved state by walking up from each effector to the root
+	// walk upwards from each effector to the root to initialize "Bone.IsSolved"
 	for (FEffector& Effector : Effectors)
 	{
 		FBone* NextBone = Effector.Bone;
@@ -391,7 +529,7 @@ bool FPBIKSolver::InitBodies()
 			FBone* BodyBone = NextBone->bIsSolverRoot ? NextBone : NextBone->Parent;
 			if (!BodyBone)
 			{
-				UE_LOG(LogTemp, Warning, TEXT("PBIK: effector is on bone that is not on or below root bone."));
+				UE_LOG(LogPBIKSolver, Warning, TEXT("PBIK: effector is on bone that is not on or below root bone."));
 				return false;
 			}
 
@@ -474,7 +612,7 @@ bool FPBIKSolver::InitConstraints()
 		FBone* BodyBone = Effector.Bone->bIsSolverRoot ? Effector.Bone : Effector.Bone->Parent;
 		if (!BodyBone)
 		{
-			UE_LOG(LogTemp, Warning, TEXT("PBIK: effector is on bone that does not have a parent."));
+			UE_LOG(LogPBIKSolver, Warning, TEXT("PBIK: effector is on bone that does not have a parent."));
 			return false;
 		}
 
@@ -582,7 +720,7 @@ PBIK::FBoneSettings* FPBIKSolver::GetBoneSettings(const int32 Index)
 
 	if (!Bones[Index].Body)
 	{
-		UE_LOG(LogTemp, Warning, TEXT("PBIK: trying to apply Bone Settings to bone that is not simulated (not between root and effector)."));
+		UE_LOG(LogPBIKSolver, Warning, TEXT("PBIK: trying to apply Bone Settings to bone, %s, that is not simulated (not between root and effector)."), *Bones[Index].Name.ToString());
 		return nullptr;
 	}
 
@@ -599,15 +737,7 @@ void FPBIKSolver::GetBoneGlobalTransform(const int32 Index, FTransform& OutTrans
 
 int32 FPBIKSolver::GetBoneIndex(FName BoneName) const
 {
-	for (int32 B=0; B < Bones.Num(); ++B)
-	{
-		if (Bones[B].Name == BoneName)
-		{
-			return B;
-		}
-	}
-
-	return INDEX_NONE;
+	return Bones.IndexOfByPredicate([&BoneName](const PBIK::FBone& Bone) { return Bone.Name == BoneName; });
 }
 
 void FPBIKSolver::SetEffectorGoal(
@@ -615,8 +745,10 @@ void FPBIKSolver::SetEffectorGoal(
 	const FVector& InPosition,
 	const FQuat& InRotation,
 	const float OffsetAlpha,
-	const float StrengthAlpha)
+	const float StrengthAlpha,
+	const float PullChainAlpha,
+	const bool bPinRotation)
 {
 	check(Index >= 0 && Index < Effectors.Num());
-	Effectors[Index].SetGoal(InPosition, InRotation, OffsetAlpha, StrengthAlpha);
+	Effectors[Index].SetGoal(InPosition, InRotation, OffsetAlpha, StrengthAlpha, PullChainAlpha, bPinRotation);
 }
diff --git a/Engine/Plugins/Experimental/FullBodyIK/Source/PBIK/Private/RigUnit_PBIK.cpp b/Engine/Plugins/Experimental/FullBodyIK/Source/PBIK/Private/RigUnit_PBIK.cpp
--- a/Engine/Plugins/Experimental/FullBodyIK/Source/PBIK/Private/RigUnit_PBIK.cpp
+++ b/Engine/Plugins/Experimental/FullBodyIK/Source/PBIK/Private/RigUnit_PBIK.cpp
@@ -3,8 +3,6 @@
 #include "RigUnit_PBIK.h"
 #include "Units/RigUnitContext.h"
 
-//#pragma optimize("", off)
-
 FRigUnit_PBIK_Execute()
 {
 	DECLARE_SCOPE_HIERARCHICAL_COUNTER_RIGUNIT()
@@ -140,7 +138,14 @@ FRigUnit_PBIK_Execute()
 		const FPBIKEffector& Effector = Effectors[E];
 		FVector Position = Effector.Transform.GetLocation();
 		FQuat Rotation = Effector.Transform.GetRotation();
-		Solver.SetEffectorGoal(EffectorSolverIndices[E], Position, Rotation, Effector.OffsetAlpha, Effector.StrengthAlpha);
+		Solver.SetEffectorGoal(
+			EffectorSolverIndices[E],
+			Position,
+			Rotation,
+			Effector.OffsetAlpha,
+			Effector.StrengthAlpha,
+			Effector.PullChainAlpha,
+			Effector.bPinRotation);
 	}
 
 	// solve
diff --git a/Engine/Plugins/Experimental/FullBodyIK/Source/PBIK/Public/Core/PBIKBody.h b/Engine/Plugins/Experimental/FullBodyIK/Source/PBIK/Public/Core/PBIKBody.h
--- a/Engine/Plugins/Experimental/FullBodyIK/Source/PBIK/Public/Core/PBIKBody.h
+++ b/Engine/Plugins/Experimental/FullBodyIK/Source/PBIK/Public/Core/PBIKBody.h
@@ -19,9 +19,6 @@ struct FBone
 {
 	FName Name;
 	int32 ParentIndex = -2; // -2 is unset, -1 for root, or 0...n otherwise
-	bool bIsSolverRoot = false;
-	bool bIsSolved = false;
-	bool bIsSubRoot = false;
 	FVector Position;
 	FQuat Rotation;
 	FVector LocalPositionOrig;
@@ -31,6 +28,10 @@ struct FBone
 	FRigidBody* Body = nullptr;
 	FBone* Parent = nullptr;
 	TArray<FBone*> Children;
+	bool bIsSolverRoot = false;
+	bool bIsSolved = false;
+	bool bIsSubRoot = false;
+	float Length = 0.f;
 	// initialized
 
 	FBone(
@@ -41,6 +42,7 @@ struct FBone
 		bool bInIsSolverRoot);
 
 	bool HasChild(const FBone* Bone);
+	void UpdateFromInputs();
 };
 
 enum class ELimitType : uint8
@@ -80,14 +82,15 @@ struct FRigidBody
 
 	FVector Position;
 	FQuat Rotation;
-	FQuat RotationOrig;
+	FVector InputPosition;
+	FQuat InitialRotation;
 	FVector BoneLocalPosition;
 	TArray<FVector> ChildLocalPositions;
 
-	float InvMass = 0.0f;
-	float MaxInvMass = 0.0f;
-	float MinInvMass = 0.0f;
-	float Length;
+	float InvMass = 0.f;
+	float MaxInvMass = 0.f;
+	float MinInvMass = 0.f;
+	float Mass = 0.f;
 	
 private:
 
@@ -109,7 +112,7 @@ public:
 	
 	void ApplyPushToPosition(const FVector& Push);
 
-	void ApplyRotationDelta(const FQuat& InDelta, const bool bNegated);
+	void ApplyRotationDelta(const FQuat& InDelta);
 };
 
 // for sorting Bodies hierarchically (root to leaf order)
diff --git a/Engine/Plugins/Experimental/FullBodyIK/Source/PBIK/Public/Core/PBIKConstraint.h b/Engine/Plugins/Experimental/FullBodyIK/Source/PBIK/Public/Core/PBIKConstraint.h
--- a/Engine/Plugins/Experimental/FullBodyIK/Source/PBIK/Public/Core/PBIKConstraint.h
+++ b/Engine/Plugins/Experimental/FullBodyIK/Source/PBIK/Public/Core/PBIKConstraint.h
@@ -59,7 +59,7 @@ private:
 
 	FVector GetPositionCorrection(FVector& OutBodyToA, FVector& OutBodyToB) const;
 
-	void ApplyRotationCorrection(FQuat PureRotA, FQuat PureRotB) const;
+	void ApplyRotationCorrection(FQuat DeltaQ) const;
 
 	void UpdateJointLimits();
 
diff --git a/Engine/Plugins/Experimental/FullBodyIK/Source/PBIK/Public/Core/PBIKSolver.h b/Engine/Plugins/Experimental/FullBodyIK/Source/PBIK/Public/Core/PBIKSolver.h
--- a/Engine/Plugins/Experimental/FullBodyIK/Source/PBIK/Public/Core/PBIKSolver.h
+++ b/Engine/Plugins/Experimental/FullBodyIK/Source/PBIK/Public/Core/PBIKSolver.h
@@ -16,9 +16,9 @@ namespace PBIK
 static float GLOBAL_UNITS = 100.0f; // (1.0f = meters), (100.0f = centimeters)
 
 // A long tail ease out function. Input range, 0-1. 
-FORCEINLINE static void QuarticEaseOut(float& InOut){ InOut = (FMath::Pow(InOut-1.0f, 4.0f) * -1.0f) + 1.0f; };
+FORCEINLINE static float QuarticEaseOut(const float& Input){ return (FMath::Pow(Input-1.0f, 4.0f) * -1.0f) + 1.0f; };
 // An ease out function. Input range, 0-1.
-FORCEINLINE static void SquaredEaseOut(float& InOut){ InOut = (FMath::Pow(InOut-1.0f, 2.0f) * -1.0f) + 1.0f; };
+FORCEINLINE static float SquaredEaseOut(const float& Input){ return (FMath::Pow(Input-1.0f, 2.0f) * -1.0f) + 1.0f; };
 
 struct FRigidBody;
 
@@ -35,11 +35,17 @@ struct FEffector
 
 	FBone* Bone;
 	TWeakPtr<FPinConstraint> Pin;
+	bool bPinRotation;
 	FRigidBody* ParentSubRoot = nullptr;
-	float LengthOfChainInInputPose;
+	float DistanceToSubRootInInputPose;
+	float DistToRootAlongBones;
 
 	float TransformAlpha;
 	float StrengthAlpha;
+	
+	float PullChainAlpha;
+	TArray<float> DistancesFromEffector;
+	float DistToSubRootAlongBones;
 
 	FEffector(FBone* InBone);
 
@@ -47,10 +53,11 @@ struct FEffector
 		const FVector& InPositionGoal,
 		const FQuat& InRotationGoal,
 		float InTransformAlpha,
-		float InStrengthAlpha);
-
-	void UpdateFromInputs();
+		float InStrengthAlpha,
+		float InPullChainAlpha,
+		bool bInPinRotation);
 
+	void UpdateFromInputs(const FBone& SolverRoot);
 	void ApplyPreferredAngles();
 };
 
@@ -81,6 +88,10 @@ struct PBIK_API FPBIKSolverSettings
 	UPROPERTY(EditAnywhere, Category = SolverSettings)
 	bool bPinRoot = false;
 
+	/** If true, solver will kinematically rotate skeleton limbs to approximate a solved pose before running constraint iterations. This can decrease the amount of iterations needed to achieve a converged pose. Default is true. */
+	UPROPERTY(EditAnywhere, Category = SolverSettings)
+	bool bPreProcessPose = true;
+
 	/** When true, the solver is reset each tick to start from the current input pose. If false, incoming animated poses are ignored and the solver starts from the results of the previous solve. Default is true. */
 	UPROPERTY(EditAnywhere, Category = SolverSettings)
 	bool bStartSolveFromInputPose = true;
@@ -120,7 +131,9 @@ public:
 		const FVector& InPosition, 
 		const FQuat& InRotation, 
 		const float OffsetAlpha, 
-		const float StrengthAlpha);
+		const float StrengthAlpha,
+		const float PullChainAlpha,
+		const bool bPinRotation);
 
 	void GetBoneGlobalTransform(const int32 Index, FTransform& OutTransform);
 
@@ -151,6 +164,14 @@ private:
 
 	void AddBodyForBone(PBIK::FBone* Bone);
 
+	void UpdateBodies(const FPBIKSolverSettings& Settings);
+
+	void SolveConstraints(const int32 Iterations, const bool bMoveRoots, const bool bAllowStretch);
+	
+	void PullRootTowardsEffectors();
+	
+	void PullChainsTowardsEffectors();
+
 private:
 
 	PBIK::FBone* SolverRoot = nullptr;
diff --git a/Engine/Plugins/Experimental/FullBodyIK/Source/PBIK/Public/RigUnit_PBIK.h b/Engine/Plugins/Experimental/FullBodyIK/Source/PBIK/Public/RigUnit_PBIK.h
--- a/Engine/Plugins/Experimental/FullBodyIK/Source/PBIK/Public/RigUnit_PBIK.h
+++ b/Engine/Plugins/Experimental/FullBodyIK/Source/PBIK/Public/RigUnit_PBIK.h
@@ -63,6 +63,12 @@ struct FPBIKEffector
 
 	UPROPERTY()
 	float StrengthAlpha = 1.0f;
+	
+	UPROPERTY()
+	float PullChainAlpha = 1.0f;
+
+	UPROPERTY()
+	bool bPinRotation = true;
 };
 
 USTRUCT(meta=(DisplayName="Full Body IK", Category="Hierarchy", Keywords="Position Based, PBIK, IK, Full Body, Multi, Effector, N-Chain, FB"))
