Merging IKRig plugin from UE5 Main

diff --git a/Engine/Plugins/Experimental/FullBodyIK/Source/PBIK/Private/Core/PBIKSolver.cpp b/Engine/Plugins/Experimental/FullBodyIK/Source/PBIK/Private/Core/PBIKSolver.cpp
--- a/Engine/Plugins/Experimental/FullBodyIK/Source/PBIK/Private/Core/PBIKSolver.cpp
+++ b/Engine/Plugins/Experimental/FullBodyIK/Source/PBIK/Private/Core/PBIKSolver.cpp
@@ -50,6 +50,8 @@ namespace PBIK
 			DistanceToSubRootInInputPose = (ParentSubRootPosition - Bone->Position).Size();
 		}
 
+		DistToRootStraightLine = (PositionOrig - SolverRoot.Position).Size();
+
 		// update distances to root (along bones)
 		{
 			if (Bone->bIsSolverRoot)
@@ -117,7 +119,7 @@ namespace PBIK
 		float DeltaSquash = DistanceToSubRootInInputPose - DistToParentSubRoot;
 		DeltaSquash = DeltaSquash > ScaledDistOrig ? ScaledDistOrig : DeltaSquash;
 		float SquashPercent = DeltaSquash / ScaledDistOrig;
-		PBIK::QuarticEaseOut(SquashPercent);
+		SquashPercent = PBIK::CircularEaseOut(SquashPercent);
 		if (SquashPercent < 0.0001f)
 		{
 			return; // limb not squashed enough
@@ -242,8 +244,8 @@ void FPBIKSolver::PullRootTowardsEffectors()
 	for (FEffector& Effector : Effectors)
 	{
 		const float DistToRootStraightLine = (Effector.Position - SolverRoot->Position).Size();
-		const float Delta = DistToRootStraightLine - Effector.DistToRootAlongBones;
-		if (Delta < 0.0f)
+		const float Delta = DistToRootStraightLine - Effector.DistToRootStraightLine;//Effector.DistToRootAlongBones;
+		if (Delta < SMALL_NUMBER)
 		{
 			continue; // only pull
 		}
diff --git a/Engine/Plugins/Experimental/FullBodyIK/Source/PBIK/Public/Core/PBIKSolver.h b/Engine/Plugins/Experimental/FullBodyIK/Source/PBIK/Public/Core/PBIKSolver.h
--- a/Engine/Plugins/Experimental/FullBodyIK/Source/PBIK/Public/Core/PBIKSolver.h
+++ b/Engine/Plugins/Experimental/FullBodyIK/Source/PBIK/Public/Core/PBIKSolver.h
@@ -17,6 +17,7 @@ static float GLOBAL_UNITS = 100.0f; // (1.0f = meters), (100.0f = centimeters)
 
 // A long tail ease out function. Input range, 0-1. 
 FORCEINLINE static float QuarticEaseOut(const float& Input){ return (FMath::Pow(Input-1.0f, 4.0f) * -1.0f) + 1.0f; };
+FORCEINLINE static float CircularEaseOut(const float& Input){ return FMath::Sqrt(1.0f - FMath::Pow(Input - 1.0f, 2.0f)); };
 // An ease out function. Input range, 0-1.
 FORCEINLINE static float SquaredEaseOut(const float& Input){ return (FMath::Pow(Input-1.0f, 2.0f) * -1.0f) + 1.0f; };
 
@@ -39,6 +40,7 @@ struct FEffector
 	FRigidBody* ParentSubRoot = nullptr;
 	float DistanceToSubRootInInputPose;
 	float DistToRootAlongBones;
+	float DistToRootStraightLine;
 
 	float TransformAlpha;
 	float StrengthAlpha;
