FBIK improvements:
- using circular ease out function to blend in preferred angles
- running stretch constraint pass AFTER preferred angles
- running Preferred Angles BEFORE Pull Chain
- refactored solve stages into separate function to improve readability
Body Mover improvements:
- added separate pos/neg weights for each translational dimension

diff --git a/Engine/Plugins/Animation/IKRig/Source/IKRig/Private/Solvers/IKRig_BodyMover.cpp b/Engine/Plugins/Animation/IKRig/Source/IKRig/Private/Solvers/IKRig_BodyMover.cpp
--- a/Engine/Plugins/Animation/IKRig/Source/IKRig/Private/Solvers/IKRig_BodyMover.cpp
+++ b/Engine/Plugins/Animation/IKRig/Source/IKRig/Private/Solvers/IKRig_BodyMover.cpp
@@ -92,8 +92,12 @@ void UIKRig_BodyMover::Solve(FIKRigSkeleton& IKRigSkeleton, const FIKRigGoalCont
 	ExtractRotation(DX, DY, DZ, RotationOffset, 50);
 
 	// alpha blend the position offset and add it to the current bone location
-	const FVector TargetPosition = (CurrentCentroid - InitialCentroid) * PositionAlpha;
-	CurrentBodyTransform.AddToTranslation(TargetPosition);
+	FVector Offset = (CurrentCentroid - InitialCentroid);
+	FVector Weight(
+		Offset.X > 0.f ? PositionPositiveX : PositionNegativeX,
+		Offset.Y > 0.f ? PositionPositiveY : PositionNegativeY,
+		Offset.Z > 0.f ? PositionPositiveZ : PositionNegativeZ);
+	CurrentBodyTransform.AddToTranslation(Offset * (Weight*PositionAlpha));
 
 	// do per-axis alpha blend
 	FVector Euler = RotationOffset.Euler() * FVector(RotateXAlpha, RotateYAlpha, RotateZAlpha);
@@ -140,6 +144,12 @@ void UIKRig_BodyMover::UpdateSolverSettings(UIKRigSolver* InSettings)
 	{
 		// copy solver settings
 		PositionAlpha = Settings->PositionAlpha;
+		PositionPositiveX = Settings->PositionPositiveX;
+		PositionPositiveY = Settings->PositionPositiveY;
+		PositionPositiveZ = Settings->PositionPositiveZ;
+		PositionNegativeX = Settings->PositionNegativeX;
+		PositionNegativeY = Settings->PositionNegativeY;
+		PositionNegativeZ = Settings->PositionNegativeZ;
 		RotationAlpha = Settings->RotationAlpha;
 		RotateXAlpha = Settings->RotateXAlpha;
 		RotateYAlpha = Settings->RotateYAlpha;
diff --git a/Engine/Plugins/Animation/IKRig/Source/IKRig/Private/Solvers/IKRig_PBIKSolver.cpp b/Engine/Plugins/Animation/IKRig/Source/IKRig/Private/Solvers/IKRig_PBIKSolver.cpp
--- a/Engine/Plugins/Animation/IKRig/Source/IKRig/Private/Solvers/IKRig_PBIKSolver.cpp
+++ b/Engine/Plugins/Animation/IKRig/Source/IKRig/Private/Solvers/IKRig_PBIKSolver.cpp
@@ -109,7 +109,7 @@ void UIKRigPBIKSolver::Solve(FIKRigSkeleton& IKRigSkeleton, const FIKRigGoalCont
 	Settings.MinMassMultiplier = MinMassMultiplier;
 	Settings.bPinRoot = bPinRoot;
 	Settings.bStartSolveFromInputPose = bStartSolveFromInputPose;
-	Settings.bPreProcessPose = bPreProcessPose;
+	Settings.bPrePullRoot = bPrePullRoot;
 
 	// solve
 	Solver.Solve(Settings);
@@ -133,7 +133,7 @@ void UIKRigPBIKSolver::UpdateSolverSettings(UIKRigSolver* InSettings)
 		MinMassMultiplier = Settings->MinMassMultiplier;
 		bPinRoot = Settings->bPinRoot;
 		bStartSolveFromInputPose = Settings->bStartSolveFromInputPose;
-		bPreProcessPose = Settings->bPreProcessPose;
+		bPrePullRoot = Settings->bPrePullRoot;
 
 		// copy effector settings
 		for (const UIKRig_FBIKEffector* InEffector : Settings->Effectors)
diff --git a/Engine/Plugins/Animation/IKRig/Source/IKRig/Public/Solvers/IKRig_BodyMover.h b/Engine/Plugins/Animation/IKRig/Source/IKRig/Public/Solvers/IKRig_BodyMover.h
--- a/Engine/Plugins/Animation/IKRig/Source/IKRig/Public/Solvers/IKRig_BodyMover.h
+++ b/Engine/Plugins/Animation/IKRig/Source/IKRig/Public/Solvers/IKRig_BodyMover.h
@@ -39,23 +39,47 @@ public:
 	FName BodyBone;
 
 	/** Blend the translational effect of this solver on/off. Range is 0-1. Default is 1.0. */
-	UPROPERTY(EditAnywhere, Category = "Body Mover Solver", meta = (ClampMin = "0", ClampMa = "1", UIMin = "0.0", UIMax = "1.0"))
+	UPROPERTY(EditAnywhere, Category = "Body Mover Solver", meta = (UIMin = "0.0", UIMax = "1.0"))
 	float PositionAlpha = 1.0f;
 
+	/** Multiply the POSITIVE X translation. Range is 0-1. Default is 1.0. */
+	UPROPERTY(EditAnywhere, Category = "Body Mover Solver", meta = (UIMin = "0.0", UIMax = "1.0"))
+	float PositionPositiveX = 1.0f;
+
+	/** Multiply the NEGATIVE X translation. Range is 0-1. Default is 1.0. */
+	UPROPERTY(EditAnywhere, Category = "Body Mover Solver", meta = (UIMin = "0.0", UIMax = "1.0"))
+	float PositionNegativeX = 1.0f;
+
+	/** Multiply the POSITIVE Y translation. Range is 0-1. Default is 1.0. */
+	UPROPERTY(EditAnywhere, Category = "Body Mover Solver", meta = (UIMin = "0.0", UIMax = "1.0"))
+	float PositionPositiveY = 1.0f;
+
+	/** Multiply the NEGATIVE Y translation. Range is 0-1. Default is 1.0. */
+	UPROPERTY(EditAnywhere, Category = "Body Mover Solver", meta = (UIMin = "0.0", UIMax = "1.0"))
+	float PositionNegativeY = 1.0f;
+
+	/** Multiply the POSITIVE Z translation. Range is 0-1. Default is 1.0. */
+	UPROPERTY(EditAnywhere, Category = "Body Mover Solver", meta = (UIMin = "0.0", UIMax = "1.0"))
+	float PositionPositiveZ = 1.0f;
+
+	/** Multiply the NEGATIVE Z translation. Range is 0-1. Default is 1.0. */
+	UPROPERTY(EditAnywhere, Category = "Body Mover Solver", meta = (UIMin = "0.0", UIMax = "1.0"))
+	float PositionNegativeZ = 1.0f;
+
 	/** Blend the total rotational effect on/off. Range is 0-1. Default is 1.0. */
-	UPROPERTY(EditAnywhere, Category = "Body Mover Solver", meta = (ClampMin = "0", ClampMa = "1", UIMin = "0.0", UIMax = "1.0"))
+	UPROPERTY(EditAnywhere, Category = "Body Mover Solver", meta = (UIMin = "0.0", UIMax = "1.0"))
 	float RotationAlpha = 1.0f;
 
 	/** Blend the X-axis rotational effect on/off. Range is 0-1. Default is 1.0. */
-	UPROPERTY(EditAnywhere, Category = "Body Mover Solver", meta = (ClampMin = "0", ClampMa = "1", UIMin = "0.0", UIMax = "1.0"))
+	UPROPERTY(EditAnywhere, Category = "Body Mover Solver", meta = (UIMin = "0.0", UIMax = "1.0"))
 	float RotateXAlpha = 1.0f;
 
 	/** Blend the Y-axis rotational effect on/off. Range is 0-1. Default is 1.0. */
-	UPROPERTY(EditAnywhere, Category = "Body Mover Solver", meta = (ClampMin = "0", ClampMa = "1", UIMin = "0.0", UIMax = "1.0"))
+	UPROPERTY(EditAnywhere, Category = "Body Mover Solver", meta = (UIMin = "0.0", UIMax = "1.0"))
 	float RotateYAlpha = 1.0f;
 
 	/** Blend the Z-axis rotational effect on/off. Range is 0-1. Default is 1.0. */
-	UPROPERTY(EditAnywhere, Category = "Body Mover Solver", meta = (ClampMin = "0", ClampMa = "1", UIMin = "0.0", UIMax = "1.0"))
+	UPROPERTY(EditAnywhere, Category = "Body Mover Solver", meta = (UIMin = "0.0", UIMax = "1.0"))
 	float RotateZAlpha = 1.0f;
 	
 	UPROPERTY()
diff --git a/Engine/Plugins/Animation/IKRig/Source/IKRig/Public/Solvers/IKRig_PBIKSolver.h b/Engine/Plugins/Animation/IKRig/Source/IKRig/Public/Solvers/IKRig_PBIKSolver.h
--- a/Engine/Plugins/Animation/IKRig/Source/IKRig/Public/Solvers/IKRig_PBIKSolver.h
+++ b/Engine/Plugins/Animation/IKRig/Source/IKRig/Public/Solvers/IKRig_PBIKSolver.h
@@ -153,9 +153,9 @@ public:
 	UPROPERTY(EditAnywhere, Category = SolverSettings)
 	bool bAllowStretch = false;
 
-	/** If true, solver will pre-process the skeleton to approximate a solved pose before running constraint iterations. Default is true. */
+	/** If true, solver will pull entire skeleton towards effectors BEFORE running constraint iterations. Default is true. */
 	UPROPERTY(EditAnywhere, Category = SolverSettings)
-	bool bPreProcessPose = true;
+	bool bPrePullRoot = true;
 
 	/** Lock the position and rotation of the solver root bone in-place (at animated position). Useful for partial-body solves. Default is false. */
 	UPROPERTY(EditAnywhere, Category = SolverSettings)
diff --git a/Engine/Plugins/Experimental/FullBodyIK/Source/PBIK/Private/Core/PBIKSolver.cpp b/Engine/Plugins/Experimental/FullBodyIK/Source/PBIK/Private/Core/PBIKSolver.cpp
--- a/Engine/Plugins/Experimental/FullBodyIK/Source/PBIK/Private/Core/PBIKSolver.cpp
+++ b/Engine/Plugins/Experimental/FullBodyIK/Source/PBIK/Private/Core/PBIKSolver.cpp
@@ -112,15 +112,13 @@
 		{
 			return; // limb is stretched
 		}
-
-		// shrink distance to reach full blend to preferred angle
-		const float ScaledDistOrig = DistanceToSubRootInInputPose;// * 0.3f;
+		
 		// amount squashed (clamped to scaled original length)
 		float DeltaSquash = DistanceToSubRootInInputPose - DistToParentSubRoot;
-		DeltaSquash = DeltaSquash > ScaledDistOrig ? ScaledDistOrig : DeltaSquash;
-		float SquashPercent = DeltaSquash / ScaledDistOrig;
+		DeltaSquash = DeltaSquash > DistanceToSubRootInInputPose ? DistanceToSubRootInInputPose : DeltaSquash;
+		float SquashPercent = DeltaSquash / DistanceToSubRootInInputPose;
 		SquashPercent = PBIK::CircularEaseOut(SquashPercent);
-		if (SquashPercent < 0.0001f)
+		if (FMath::IsNearlyZero(SquashPercent))
 		{
 			return; // limb not squashed enough
 		}
@@ -189,49 +187,8 @@
 	// do all constraint solving
 	UpdateBodies(Settings);
 
-	// update Bone transforms controlled by Bodies
-	for (FRigidBody& Body : Bodies)
-	{
-		Body.Bone->Position = Body.Position + Body.Rotation * Body.BoneLocalPosition;
-		Body.Bone->Rotation = Body.Rotation;
-	}
-
-	// update Bone transforms controlled by effectors
-	for (const FEffector& Effector : Effectors)
-	{
-		FBone* Bone = Effector.Bone;
-		if (Bone->bIsSolverRoot)
-		{
-			continue; // if there's an effector on the root, leave it where the body ended up
-		}
-
-		if (Bone->Body)
-		{
-			continue; // effector is between other effectors, so leave transform where body ended up
-		}
-		
-		Bone->Position = Bone->Parent->Position + Bone->Parent->Rotation * Bone->LocalPositionOrig;
-
-		if (Effector.bPinRotation)
-		{
-			Bone->Rotation = Effector.Rotation; // optionally pin rotation to that of effector
-		}else
-		{
-			Bone->Rotation = Bone->Parent->Rotation * Bone->LocalRotationOrig;
-		}
-	}
-
-	// propagate to non-solved bones (requires storage in root to tip order)
-	for (FBone& Bone : Bones)
-	{
-		if (Bone.bIsSolved || !Bone.Parent)
-		{
-			continue;
-		}
-
-		Bone.Position = Bone.Parent->Position + Bone.Parent->Rotation * Bone.LocalPositionOrig;
-		Bone.Rotation = Bone.Parent->Rotation * Bone.LocalRotationOrig;
-	}
+	// now that bodies are posed, update bone transforms
+	UpdateBonesFromBodies();
 }
 
 void FPBIKSolver::PullRootTowardsEffectors()
@@ -264,7 +221,7 @@
 	}
 }
 
-void FPBIKSolver::PullChainsTowardsEffectors()
+void FPBIKSolver::ApplyPullChainAlpha()
 {
 	using PBIK::FEffector;
 	using PBIK::FBone;
@@ -282,6 +239,11 @@
 			continue;
 		}
 
+		if (Effector.PullChainAlpha < SMALL_NUMBER)
+		{
+			continue;
+		}
+
 		// get original chain vector
 		const FVector ChainStartOrig = Effector.ParentSubRoot->InputPosition;
 		const FVector ChainEndOrig = Effector.PositionOrig;
@@ -331,27 +293,38 @@
 	}
 }
 
-void FPBIKSolver::UpdateBodies(const FPBIKSolverSettings& Settings)
+void FPBIKSolver::ApplyPreferredAngles()
 {
-	using PBIK::FRigidBody;
+	// apply preferred angles to squashed effector chains
 	using PBIK::FEffector;
+	for (FEffector& Effector : Effectors)
+	{
+		Effector.ApplyPreferredAngles();
+	}
 
+	// preferred angles introduce stretch, so remove that to prevent biasing the first constraint iteration
+	for (int32 C = Constraints.Num() - 1; C >= 0; --C)
+	{
+		Constraints[C]->RemoveStretch();
+	}
+}
+
+void FPBIKSolver::UpdateBodies(const FPBIKSolverSettings& Settings)
+{
 	// this creates a better pose from which to start the constraint solving
-	if (Settings.bPreProcessPose)
+	if (Settings.bPrePullRoot)
 	{
 		// pull ALL bodies towards the effectors
 		PullRootTowardsEffectors();
-		// pull each effector's chain towards itself
-		PullChainsTowardsEffectors();
 	}
-	
-	// apply preferred angles to squashed effector chains
-	for (FEffector& Effector : Effectors)
-	{
-		Effector.ApplyPreferredAngles();
-	}
-	
-	// run constraint iterations while allowing stretch, just to get reaching pose
+
+	// pre-rotate bones towards preferred angles when effector limb is squashed
+	ApplyPreferredAngles();
+
+	// pull each effector's chain towards itself
+	ApplyPullChainAlpha();
+		
+	// run ALL constraint iterations
 	SolveConstraints(Settings.Iterations, true, Settings.bAllowStretch);
 }
 
@@ -387,6 +360,57 @@
 	}
 }
 
+void FPBIKSolver::UpdateBonesFromBodies()
+{
+	using PBIK::FRigidBody;
+	using PBIK::FBone;
+	using PBIK::FEffector;
+	
+	// update Bone transforms controlled by Bodies
+	for (FRigidBody& Body : Bodies)
+	{
+		Body.Bone->Position = Body.Position + Body.Rotation * Body.BoneLocalPosition;
+		Body.Bone->Rotation = Body.Rotation;
+	}
+
+	// update Bone transforms controlled by effectors
+	for (const FEffector& Effector : Effectors)
+	{
+		FBone* Bone = Effector.Bone;
+		if (Bone->bIsSolverRoot)
+		{
+			continue; // if there's an effector on the root, leave it where the body ended up
+		}
+
+		if (Bone->Body)
+		{
+			continue; // effector is between other effectors, so leave transform where body ended up
+		}
+		
+		Bone->Position = Bone->Parent->Position + Bone->Parent->Rotation * Bone->LocalPositionOrig;
+
+		if (Effector.bPinRotation)
+		{
+			Bone->Rotation = Effector.Rotation; // optionally pin rotation to that of effector
+		}else
+		{
+			Bone->Rotation = Bone->Parent->Rotation * Bone->LocalRotationOrig;
+		}
+	}
+
+	// propagate to non-solved bones (requires storage in root to tip order)
+	for (FBone& Bone : Bones)
+	{
+		if (Bone.bIsSolved || !Bone.Parent)
+		{
+			continue;
+		}
+
+		Bone.Position = Bone.Parent->Position + Bone.Parent->Rotation * Bone.LocalPositionOrig;
+		Bone.Rotation = Bone.Parent->Rotation * Bone.LocalRotationOrig;
+	}
+}
+
 bool FPBIKSolver::Initialize()
 {
 	if (bReadyToSimulate)
diff --git a/Engine/Plugins/Experimental/FullBodyIK/Source/PBIK/Public/Core/PBIKSolver.h b/Engine/Plugins/Experimental/FullBodyIK/Source/PBIK/Public/Core/PBIKSolver.h
--- a/Engine/Plugins/Experimental/FullBodyIK/Source/PBIK/Public/Core/PBIKSolver.h
+++ b/Engine/Plugins/Experimental/FullBodyIK/Source/PBIK/Public/Core/PBIKSolver.h
@@ -17,6 +17,7 @@ static float GLOBAL_UNITS = 100.0f; // (1.0f = meters), (100.0f = centimeters)
 
 // A long tail ease out function. Input range, 0-1. 
 FORCEINLINE static float QuarticEaseOut(const float& Input){ return (FMath::Pow(Input-1.0f, 4.0f) * -1.0f) + 1.0f; };
+// A C2 continuous ease out function. Input range, 0-1
 FORCEINLINE static float CircularEaseOut(const float& Input){ return FMath::Sqrt(1.0f - FMath::Pow(Input - 1.0f, 2.0f)); };
 // An ease out function. Input range, 0-1.
 FORCEINLINE static float SquaredEaseOut(const float& Input){ return (FMath::Pow(Input-1.0f, 2.0f) * -1.0f) + 1.0f; };
@@ -92,7 +93,7 @@ struct PBIK_API FPBIKSolverSettings
 
 	/** If true, solver will kinematically rotate skeleton limbs to approximate a solved pose before running constraint iterations. This can decrease the amount of iterations needed to achieve a converged pose. Default is true. */
 	UPROPERTY(EditAnywhere, Category = SolverSettings)
-	bool bPreProcessPose = true;
+	bool bPrePullRoot = true;
 
 	/** When true, the solver is reset each tick to start from the current input pose. If false, incoming animated poses are ignored and the solver starts from the results of the previous solve. Default is true. */
 	UPROPERTY(EditAnywhere, Category = SolverSettings)
@@ -168,11 +169,15 @@ private:
 
 	void UpdateBodies(const FPBIKSolverSettings& Settings);
 
+	void UpdateBonesFromBodies();
+
 	void SolveConstraints(const int32 Iterations, const bool bMoveRoots, const bool bAllowStretch);
 	
 	void PullRootTowardsEffectors();
 	
-	void PullChainsTowardsEffectors();
+	void ApplyPullChainAlpha();
+
+	void ApplyPreferredAngles();
 
 private:
 
@@ -186,4 +191,4 @@ private:
 	
 	PBIK::FDebugDraw DebugDraw;
 	friend PBIK::FDebugDraw;
-};
+};
\ No newline at end of file
