Make sure the FBIK plugin gets loaded before any game modules. Otherwise if it gets loaded as a part of a project, it won't be found.

#rnx

diff --git a/Engine/Plugins/Experimental/FullBodyIK/FullBodyIK.uplugin b/Engine/Plugins/Experimental/FullBodyIK/FullBodyIK.uplugin
--- a/Engine/Plugins/Experimental/FullBodyIK/FullBodyIK.uplugin
+++ b/Engine/Plugins/Experimental/FullBodyIK/FullBodyIK.uplugin
@@ -24,12 +24,12 @@
 		{
 			"Name": "FullBodyIK",
 			"Type": "Runtime",
-			"LoadingPhase": "Default"
+			"LoadingPhase": "PreDefault"
 		},
 		{
 			"Name": "PBIK",
 			"Type": "Runtime",
-			"LoadingPhase": "Default"
+			"LoadingPhase": "PreDefault"
 		}
 	]
 }
\ No newline at end of file
