IK Rig: bPinRoot FBIK solver option disabled bPullRoot in details panel.

diff --git a/Engine/Plugins/Animation/IKRig/Source/IKRig/Public/Solvers/IKRig_PBIKSolver.h b/Engine/Plugins/Animation/IKRig/Source/IKRig/Public/Solvers/IKRig_PBIKSolver.h
--- a/Engine/Plugins/Animation/IKRig/Source/IKRig/Public/Solvers/IKRig_PBIKSolver.h
+++ b/Engine/Plugins/Animation/IKRig/Source/IKRig/Public/Solvers/IKRig_PBIKSolver.h
@@ -201,7 +201,7 @@ public:
 	/** (Default is true) If true, solver runs a pre-pass (before the constraint solve) which will translate the whole body by the average motion of the stretched effectors.
 	*This can decrease the amount of iterations needed to achieve a converged pose when the effectors are pulling the body far.
 	*This option has no effect if PinRoot is set to True.*/
-	UPROPERTY(EditAnywhere, Category = "Full Body IK Settings")
+	UPROPERTY(EditAnywhere, Category = "Full Body IK Settings", meta=(EditCondition = "!bPinRoot"))
 	bool bPrePullRoot = true;
 
 	/** Lock the position and rotation of the solver root bone in-place (at animated position). Useful for partial-body solves. Default is false. */
diff --git a/Engine/Plugins/Experimental/FullBodyIK/Source/PBIK/Public/Core/PBIKSolver.h b/Engine/Plugins/Experimental/FullBodyIK/Source/PBIK/Public/Core/PBIKSolver.h
--- a/Engine/Plugins/Experimental/FullBodyIK/Source/PBIK/Public/Core/PBIKSolver.h
+++ b/Engine/Plugins/Experimental/FullBodyIK/Source/PBIK/Public/Core/PBIKSolver.h
@@ -119,7 +119,7 @@ struct PBIK_API FPBIKSolverSettings
 	/** (Default is true) If true, solver runs a pre-pass (before the constraint solve) which will translate the whole body by the average motion of the stretched effectors.
 	 *This can decrease the amount of iterations needed to achieve a converged pose when the effectors are pulling the body far.
 	 *This option has no effect if PinRoot is set to True.*/
-	UPROPERTY(EditAnywhere, Category = SolverSettings)
+	UPROPERTY(EditAnywhere, Category = SolverSettings, meta=(EditCondition = "!bPinRoot"))
 	bool bPrePullRoot = true;
 
 	/** When true, the solver is reset each tick to start from the current input pose. If false, incoming animated poses are ignored and the solver starts from the results of the previous solve. Default is true. */
