Fixing min/max range of FBIK effector settings.

diff --git a/Engine/Plugins/Experimental/FullBodyIK/Source/PBIK/Public/RigUnit_PBIK.h b/Engine/Plugins/Experimental/FullBodyIK/Source/PBIK/Public/RigUnit_PBIK.h
--- a/Engine/Plugins/Experimental/FullBodyIK/Source/PBIK/Public/RigUnit_PBIK.h
+++ b/Engine/Plugins/Experimental/FullBodyIK/Source/PBIK/Public/RigUnit_PBIK.h
@@ -74,7 +74,7 @@ struct FPBIKEffector
 	 * At 0.0, the effector does not pull at all, but the bones between the effector and the root will still slightly resist motion from other effectors.
 	 * This can thus act as a "stabilizer" of sorts for parts of the body that you do not want to behave in a pure FK fashion.
 	 */
-	UPROPERTY(EditAnywhere, Category="Effector")
+	UPROPERTY(EditAnywhere, Category="Effector", meta = (ClampMin = "0.0", ClampMax = "1.0", UIMin = "0.0", UIMax = "1.0"))
 	float StrengthAlpha = 1.0f;
 
 	/** Range 0-1 (default is 1.0). When enabled (greater than 0.0), the solver internally partitions the skeleton into 'chains' which extend from the effector to the nearest fork in the skeleton.
@@ -82,12 +82,12 @@ struct FPBIKEffector
 	 *This can improve the results for sparse bone chains, and significantly improve convergence on dense bone chains.
 	 *But it may cause undesirable results in highly constrained bone chains (like robot arms).
 	 */
-	UPROPERTY(EditAnywhere, Category="Effector")
+	UPROPERTY(EditAnywhere, Category="Effector", meta = (ClampMin = "0.0", ClampMax = "1.0", UIMin = "0.0", UIMax = "1.0"))
 	float PullChainAlpha = 1.0f;
 
 	/** Range 0-1 (default is 1.0).
 	 *Blends the effector bone rotation between the rotation of the effector transform (1.0) and the rotation of the input bone (0.0).*/
-	UPROPERTY(EditAnywhere, Category="Effector")
+	UPROPERTY(EditAnywhere, Category="Effector", meta = (ClampMin = "0.0", ClampMax = "1.0", UIMin = "0.0", UIMax = "1.0"))
 	float PinRotation = 1.0f;
 };
 
