Initialise FPBIKSolverSettings::RootBehaviour

The engine attempts to find uninitialised struct members. This check would fire due to the above member not being initialised and cause CIS' cooks of ShooterGame to output errors.

#rnx

diff --git a/Engine/Plugins/Experimental/FullBodyIK/Source/PBIK/Public/Core/PBIKSolver.h b/Engine/Plugins/Experimental/FullBodyIK/Source/PBIK/Public/Core/PBIKSolver.h
--- a/Engine/Plugins/Experimental/FullBodyIK/Source/PBIK/Public/Core/PBIKSolver.h
+++ b/Engine/Plugins/Experimental/FullBodyIK/Source/PBIK/Public/Core/PBIKSolver.h
@@ -125,7 +125,7 @@ struct PBIK_API FPBIKSolverSettings
 	*Pin to Input: locks the translation and rotation of the root bone to the input pose. Overrides any bone settings applied to the root. Good for partial-body solves.
 	*Free: treats the root bone like any other and allows it to move freely or according to any bone settings applied to it. */
 	UPROPERTY(EditAnywhere, Category = SolverSettings)
-	EPBIKRootBehavior RootBehavior;
+	EPBIKRootBehavior RootBehavior = EPBIKRootBehavior::PrePull;
 
 	/** When true, the solver is reset each tick to start from the current input pose. If false, incoming animated poses are ignored and the solver starts from the results of the previous solve. Default is true. */
 	UPROPERTY(EditAnywhere, Category = SolverSettings)
