PBIK: Minor fixes

diff --git a/Engine/Plugins/Experimental/FullBodyIK/Source/PBIK/Private/RigUnit_PBIK.cpp b/Engine/Plugins/Experimental/FullBodyIK/Source/PBIK/Private/RigUnit_PBIK.cpp
--- a/Engine/Plugins/Experimental/FullBodyIK/Source/PBIK/Private/RigUnit_PBIK.cpp
+++ b/Engine/Plugins/Experimental/FullBodyIK/Source/PBIK/Private/RigUnit_PBIK.cpp
@@ -109,6 +109,11 @@ FRigUnit_PBIK_Execute()
 		return;
 	}
 
+	if (EffectorSolverIndices.Num() != Effectors.Num())
+	{
+		return;
+	}
+
 	// set bones to input pose
 	for(int32 BoneIndex = 0; BoneIndex < Solver.GetNumBones(); BoneIndex++)
 	{
diff --git a/Engine/Plugins/Experimental/FullBodyIK/Source/PBIK/Public/RigUnit_PBIK.h b/Engine/Plugins/Experimental/FullBodyIK/Source/PBIK/Public/RigUnit_PBIK.h
--- a/Engine/Plugins/Experimental/FullBodyIK/Source/PBIK/Public/RigUnit_PBIK.h
+++ b/Engine/Plugins/Experimental/FullBodyIK/Source/PBIK/Public/RigUnit_PBIK.h
@@ -55,7 +55,7 @@ struct FPBIKEffector
 	FPBIKEffector()	: Bone(NAME_None) {}
 
 	/** The bone that this effector will pull on. */
-	UPROPERTY(EditAnywhere, Category="Effector", meta = (Constant, CustomWidget = "BoneName"))
+	UPROPERTY(EditAnywhere, Category="Effector", meta = (CustomWidget = "BoneName"))
 	FName Bone;
 
 	/** The target location and rotation for this effector. The solver will try to get the specified bone to reach this location.*/
@@ -111,7 +111,7 @@ struct FRigUnit_PBIK : public FRigUnit_HighlevelBaseMutable
 	/**This is usually the top-most skinned bone; often the "Pelvis" or "Hips", but can be set to any bone.
 	 *Bones above the root will be ignored by the solver.
 	 *Bones that are located *between* the Root and the effectors will be included in the solve.*/
-	UPROPERTY(meta = (Input, Constant, CustomWidget = "BoneName"))
+	UPROPERTY(meta = (Input, CustomWidget = "BoneName"))
 	FName Root;
 
 	/** An array of effectors. These specify target transforms for different parts of the skeleton. */
@@ -127,7 +127,7 @@ struct FRigUnit_PBIK : public FRigUnit_HighlevelBaseMutable
 
 	/** These bones will be excluded from the solver. They will not bend and will not contribute to the constraint set.
 	 * Use the ExcludedBones array instead of setting Rotation Stiffness to very high values or Rotation Limits with zero range. */
-	UPROPERTY(meta = (Input, Constant, CustomWidget = "BoneName"))
+	UPROPERTY(meta = (Input, CustomWidget = "BoneName"))
 	TArray<FName> ExcludedBones;
 
 	/** Global solver settings. */
