FBIK: Making "Pull Chain Alpha" 0 by default instead of 1.

diff --git a/Engine/Plugins/Experimental/FullBodyIK/Source/PBIK/Public/RigUnit_PBIK.h b/Engine/Plugins/Experimental/FullBodyIK/Source/PBIK/Public/RigUnit_PBIK.h
--- a/Engine/Plugins/Experimental/FullBodyIK/Source/PBIK/Public/RigUnit_PBIK.h
+++ b/Engine/Plugins/Experimental/FullBodyIK/Source/PBIK/Public/RigUnit_PBIK.h
@@ -83,7 +83,7 @@ struct FPBIKEffector
 	 *But it may cause undesirable results in highly constrained bone chains (like robot arms).
 	 */
 	UPROPERTY(EditAnywhere, Category="Effector", meta = (ClampMin = "0.0", ClampMax = "1.0", UIMin = "0.0", UIMax = "1.0"))
-	float PullChainAlpha = 1.0f;
+	float PullChainAlpha = 0.0f;
 
 	/** Range 0-1 (default is 1.0).
 	 *Blends the effector bone rotation between the rotation of the effector transform (1.0) and the rotation of the input bone (0.0).*/
