ControlRig: Remove EControlRigState from units

diff --git a/Engine/Plugins/Animation/ControlRig/Source/ControlRigEditor/Private/Editor/ControlRigEditor.h b/Engine/Plugins/Animation/ControlRig/Source/ControlRigEditor/Private/Editor/ControlRigEditor.h
--- a/Engine/Plugins/Animation/ControlRig/Source/ControlRigEditor/Private/Editor/ControlRigEditor.h
+++ b/Engine/Plugins/Animation/ControlRig/Source/ControlRigEditor/Private/Editor/ControlRigEditor.h
@@ -239,7 +239,7 @@ protected:
 
 	void HandleModifiedEvent(ERigVMGraphNotifType InNotifType, URigVMGraph* InGraph, UObject* InSubject);
 	void HandleVMCompiledEvent(UObject* InCompiledObject, URigVM* InVM);
-	void HandleControlRigExecutedEvent(UControlRig* InControlRig, const EControlRigState InState, const FName& InEventName);
+	void HandleControlRigExecutedEvent(UControlRig* InControlRig, const FName& InEventName);
 	void HandleControlRigExecutionHalted(const int32 InstructionIndex, UObject* InNode, const FName& InEntryName);
 	void SetHaltedNode(URigVMNode* Node);
 
@@ -494,8 +494,8 @@ protected:
 
 	FDelegateHandle PropertyChangedHandle;
 
-	void OnPreConstruction_AnyThread(UControlRig* InRig, const EControlRigState InState, const FName& InEventName);
-	void OnPostConstruction_AnyThread(UControlRig* InRig, const EControlRigState InState, const FName& InEventName);
+	void OnPreConstruction_AnyThread(UControlRig* InRig, const FName& InEventName);
+	void OnPostConstruction_AnyThread(UControlRig* InRig, const FName& InEventName);
 
 	bool bIsConstructionEventRunning;
 	uint32 LastHierarchyHash;
diff --git a/Engine/Plugins/Animation/ControlRig/Source/ControlRigEditor/Private/Editor/SControlRigAnimAttributeView.cpp b/Engine/Plugins/Animation/ControlRig/Source/ControlRigEditor/Private/Editor/SControlRigAnimAttributeView.cpp
--- a/Engine/Plugins/Animation/ControlRig/Source/ControlRigEditor/Private/Editor/SControlRigAnimAttributeView.cpp
+++ b/Engine/Plugins/Animation/ControlRig/Source/ControlRigEditor/Private/Editor/SControlRigAnimAttributeView.cpp
@@ -86,12 +86,11 @@ void SControlRigAnimAttributeView::StopObservingCurrentControlRig()
 
 void SControlRigAnimAttributeView::HandleControlRigPostForwardSolve(
 	UControlRig* InControlRig,
-	const EControlRigState InState,
 	const FName& InEventName) const
 {
 	if (IsValid(InControlRig) && ensure(InControlRig == ControlRigBeingDebuggedPtr))
 	{
-		if (InState == EControlRigState::Update && InEventName == FRigUnit_BeginExecution::EventName)
+		if (InEventName == FRigUnit_BeginExecution::EventName)
 		{
 			if (const USkeletalMeshComponent* SkeletalMeshComponent = Cast<USkeletalMeshComponent>(InControlRig->OuterSceneComponent.Get()))
 			{
diff --git a/Engine/Plugins/Animation/ControlRig/Source/ControlRigEditor/Private/Editor/SControlRigAnimAttributeView.h b/Engine/Plugins/Animation/ControlRig/Source/ControlRigEditor/Private/Editor/SControlRigAnimAttributeView.h
--- a/Engine/Plugins/Animation/ControlRig/Source/ControlRigEditor/Private/Editor/SControlRigAnimAttributeView.h
+++ b/Engine/Plugins/Animation/ControlRig/Source/ControlRigEditor/Private/Editor/SControlRigAnimAttributeView.h
@@ -8,7 +8,6 @@
 class FControlRigEditor;
 class UControlRigBlueprint;
 class UControlRig;
-enum class EControlRigState : uint8;
 
 class SAnimAttributeView;
 
@@ -29,7 +28,7 @@ class SControlRigAnimAttributeView : public SCompoundWidget
 	void StartObservingNewControlRig(UControlRig* InControlRig);
 	void StopObservingCurrentControlRig();
 
-	void HandleControlRigPostForwardSolve(UControlRig* InControlRig, const EControlRigState InState, const FName& InEventName) const;
+	void HandleControlRigPostForwardSolve(UControlRig* InControlRig, const FName& InEventName) const;
 	
 	TWeakPtr<FControlRigEditor> ControlRigEditor;
 	TWeakObjectPtr<UControlRigBlueprint> ControlRigBlueprint;
diff --git a/Engine/Plugins/Animation/ControlRig/Source/ControlRigEditor/Private/Editor/SControlRigStackView.cpp b/Engine/Plugins/Animation/ControlRig/Source/ControlRigEditor/Private/Editor/SControlRigStackView.cpp
--- a/Engine/Plugins/Animation/ControlRig/Source/ControlRigEditor/Private/Editor/SControlRigStackView.cpp
+++ b/Engine/Plugins/Animation/ControlRig/Source/ControlRigEditor/Private/Editor/SControlRigStackView.cpp
@@ -1038,7 +1038,7 @@ void SControlRigStackView::HandleModifiedEvent(ERigVMGraphNotifType InNotifType,
 	}
 }
 
-void SControlRigStackView::HandleControlRigInitializedEvent(UControlRig* InControlRig, const EControlRigState InState, const FName& InEventName)
+void SControlRigStackView::HandleControlRigInitializedEvent(UControlRig* InControlRig, const FName& InEventName)
 {
 	TGuardValue<bool> SuspendControllerSelection(bSuspendControllerSelection, true);
 
diff --git a/Engine/Plugins/Animation/ControlRig/Source/ControlRigEditor/Private/Editor/SControlRigStackView.h b/Engine/Plugins/Animation/ControlRig/Source/ControlRigEditor/Private/Editor/SControlRigStackView.h
--- a/Engine/Plugins/Animation/ControlRig/Source/ControlRigEditor/Private/Editor/SControlRigStackView.h
+++ b/Engine/Plugins/Animation/ControlRig/Source/ControlRigEditor/Private/Editor/SControlRigStackView.h
@@ -114,7 +114,7 @@ private:
 	bool bSuspendModelNotifications;
 	bool bSuspendControllerSelection;
 	void HandleModifiedEvent(ERigVMGraphNotifType InNotifType, URigVMGraph* InGraph, UObject* InSubject);
-	void HandleControlRigInitializedEvent(UControlRig* InControlRig, const EControlRigState InState, const FName& InEventName);
+	void HandleControlRigInitializedEvent(UControlRig* InControlRig, const FName& InEventName);
 	void HandlePreviewControlRigUpdated(FControlRigEditor* InEditor);
 	void HandleItemMouseDoubleClick(TSharedPtr<FRigStackEntry> InItem);
 
diff --git a/Engine/Plugins/Animation/ControlRig/Source/ControlRigEditor/Private/Editor/SRigHierarchy.cpp b/Engine/Plugins/Animation/ControlRig/Source/ControlRigEditor/Private/Editor/SRigHierarchy.cpp
--- a/Engine/Plugins/Animation/ControlRig/Source/ControlRigEditor/Private/Editor/SRigHierarchy.cpp
+++ b/Engine/Plugins/Animation/ControlRig/Source/ControlRigEditor/Private/Editor/SRigHierarchy.cpp
@@ -860,7 +860,7 @@ void SRigHierarchy::HandleSetObjectBeingDebugged(UObject* InObject)
 	RefreshTreeView();
 }
 
-void SRigHierarchy::OnPreConstruction_AnyThread(UControlRig* InRig, const EControlRigState InState, const FName& InEventName)
+void SRigHierarchy::OnPreConstruction_AnyThread(UControlRig* InRig, const FName& InEventName)
 {
 	if(InRig != ControlRigBeingDebuggedPtr.Get())
 	{
@@ -870,7 +870,7 @@ void SRigHierarchy::OnPreConstruction_AnyThread(UControlRig* InRig, const EContr
 	SelectionBeforeConstruction = InRig->GetHierarchy()->GetSelectedKeys();
 }
 
-void SRigHierarchy::OnPostConstruction_AnyThread(UControlRig* InRig, const EControlRigState InState, const FName& InEventName)
+void SRigHierarchy::OnPostConstruction_AnyThread(UControlRig* InRig, const FName& InEventName)
 {
 	if(InRig != ControlRigBeingDebuggedPtr.Get())
 	{
diff --git a/Engine/Plugins/Animation/ControlRig/Source/ControlRigEditor/Private/Editor/SRigHierarchy.h b/Engine/Plugins/Animation/ControlRig/Source/ControlRigEditor/Private/Editor/SRigHierarchy.h
--- a/Engine/Plugins/Animation/ControlRig/Source/ControlRigEditor/Private/Editor/SRigHierarchy.h
+++ b/Engine/Plugins/Animation/ControlRig/Source/ControlRigEditor/Private/Editor/SRigHierarchy.h
@@ -224,8 +224,8 @@ private:
 	void OnHierarchyModified_AnyThread(ERigHierarchyNotification InNotif, URigHierarchy* InHierarchy, const FRigBaseElement* InElement);
 	void HandleRefreshEditorFromBlueprint(UControlRigBlueprint* InBlueprint);
 	void HandleSetObjectBeingDebugged(UObject* InObject);
-	void OnPreConstruction_AnyThread(UControlRig* InRig, const EControlRigState InState, const FName& InEventName);
-	void OnPostConstruction_AnyThread(UControlRig* InRig, const EControlRigState InState, const FName& InEventName);
+	void OnPreConstruction_AnyThread(UControlRig* InRig, const FName& InEventName);
+	void OnPostConstruction_AnyThread(UControlRig* InRig, const FName& InEventName);
 
 	bool bIsConstructionEventRunning;
 	uint32 LastHierarchyHash;
diff --git a/Engine/Plugins/Animation/ControlRig/Source/ControlRigEditor/Private/Sequencer/ControlRigParameterTrackEditor.cpp b/Engine/Plugins/Animation/ControlRig/Source/ControlRigEditor/Private/Sequencer/ControlRigParameterTrackEditor.cpp
--- a/Engine/Plugins/Animation/ControlRig/Source/ControlRigEditor/Private/Sequencer/ControlRigParameterTrackEditor.cpp
+++ b/Engine/Plugins/Animation/ControlRig/Source/ControlRigEditor/Private/Sequencer/ControlRigParameterTrackEditor.cpp
@@ -2987,7 +2987,7 @@ void FControlRigParameterTrackEditor::HandleControlSelected(UControlRig* Subject
 
 
 
-void FControlRigParameterTrackEditor::HandleOnInitialized(UControlRig* ControlRig, const EControlRigState InState, const FName& InEventName)
+void FControlRigParameterTrackEditor::HandleOnInitialized(UControlRig* ControlRig, const FName& InEventName)
 {
 	if (GetSequencer().IsValid())
 	{
diff --git a/Engine/Plugins/Animation/ControlRig/Source/ControlRigEditor/Private/Sequencer/ControlRigParameterTrackEditor.h b/Engine/Plugins/Animation/ControlRig/Source/ControlRigEditor/Private/Sequencer/ControlRigParameterTrackEditor.h
--- a/Engine/Plugins/Animation/ControlRig/Source/ControlRigEditor/Private/Sequencer/ControlRigParameterTrackEditor.h
+++ b/Engine/Plugins/Animation/ControlRig/Source/ControlRigEditor/Private/Sequencer/ControlRigParameterTrackEditor.h
@@ -126,7 +126,7 @@ private:
 	void HandleControlModified(UControlRig* Subject, FRigControlElement* ControlElement, const FRigControlModifiedContext& Context);
 	void HandleControlSelected(UControlRig* Subject, FRigControlElement* ControlElement, bool bSelected);
 	void HandleControlUndoBracket(UControlRig* Subject, bool bOpenUndoBracket);
-	void HandleOnInitialized(UControlRig* Subject, const EControlRigState InState, const FName& InEventName);
+	void HandleOnInitialized(UControlRig* Subject, const FName& InEventName);
 	void HandleOnControlRigBound(UControlRig* InControlRig);
 	void HandleOnObjectBoundToControlRig(UObject* InObject);
 
diff --git a/Engine/Plugins/Animation/ControlRigSpline/Source/ControlRigSpline/Private/ControlRigSplineUnits.cpp b/Engine/Plugins/Animation/ControlRigSpline/Source/ControlRigSpline/Private/ControlRigSplineUnits.cpp
--- a/Engine/Plugins/Animation/ControlRigSpline/Source/ControlRigSpline/Private/ControlRigSplineUnits.cpp
+++ b/Engine/Plugins/Animation/ControlRigSpline/Source/ControlRigSpline/Private/ControlRigSplineUnits.cpp
@@ -12,39 +12,21 @@
 
 FRigUnit_ControlRigSplineFromPoints_Execute()
 {
-	switch (ExecuteContext.UnitContext.State)
-	{
-		case EControlRigState::Init:
-		case EControlRigState::Update:
-		{
-			// reset the spline
-			Spline = FControlRigSpline();
+	// reset the spline
+	Spline = FControlRigSpline();
 
-			if (Points.Num() < 4)
-			{
-				UE_CONTROLRIG_RIGUNIT_REPORT_WARNING(TEXT("Cannot create a spline with less than 4 points (%d received)."), Points.Num());
-				return;
-			}
-				
-			const TArrayView<const FVector> PointsView(Points.GetData(), Points.Num());
-			Spline.SetControlPoints(PointsView, SplineMode, SamplesPerSegment, Compression, Stretch);
-			break;
-		}
-		default:
-		{
-			checkNoEntry(); // Execute is only defined for Init and Update
-			break;
-		}
+	if (Points.Num() < 4)
+	{
+		UE_CONTROLRIG_RIGUNIT_REPORT_WARNING(TEXT("Cannot create a spline with less than 4 points (%d received)."), Points.Num());
+		return;
 	}
+		
+	const TArrayView<const FVector> PointsView(Points.GetData(), Points.Num());
+	Spline.SetControlPoints(PointsView, SplineMode, SamplesPerSegment, Compression, Stretch);
 }
 
 FRigUnit_SetSplinePoints_Execute()
 {
-	if (ExecuteContext.UnitContext.State == EControlRigState::Init)
-	{
-		return;
-	}
-	
 	if (!Spline.SplineData.IsValid())
 	{
 		UE_LOG(LogControlRig, Error, TEXT("Invalid input spline."));
@@ -65,21 +47,8 @@ FRigUnit_SetSplinePoints_Execute()
 		return;
 	}
 
-	switch (ExecuteContext.UnitContext.State)
-	{
-		case EControlRigState::Init:
-		case EControlRigState::Update:
-		{
-			const TArrayView<const FVector> PointsView(Points.GetData(), Points.Num());
-			Spline.SetControlPoints(PointsView, Spline.SplineData->SplineMode, Spline.SplineData->SamplesPerSegment, Spline.SplineData->Compression, Spline.SplineData->Stretch);
-			break;
-		}
-		default:
-		{
-			checkNoEntry(); // Execute is only defined for Init and Update
-			break;
-		}
-	}
+	const TArrayView<const FVector> PointsView(Points.GetData(), Points.Num());
+	Spline.SetControlPoints(PointsView, Spline.SplineData->SplineMode, Spline.SplineData->SamplesPerSegment, Spline.SplineData->Compression, Spline.SplineData->Stretch);
 }
 
 FRigUnit_PositionFromControlRigSpline_Execute()
@@ -97,20 +66,7 @@ FRigUnit_PositionFromControlRigSpline_Execute()
 	}
 #endif
 	
-	switch (ExecuteContext.UnitContext.State)
-	{
-		case EControlRigState::Init:
-		case EControlRigState::Update:
-		{
-			Position = Spline.PositionAtParam(U);
-			break;
-		}
-		default:
-		{
-			checkNoEntry(); // Execute is only defined for Init and Update
-			break;
-		}
-	}
+	Position = Spline.PositionAtParam(U);
 }
 
 FRigUnit_TransformFromControlRigSpline_Execute()
@@ -128,39 +84,26 @@ FRigUnit_TransformFromControlRigSpline_Execute()
 	}
 #endif	
 	
-	switch (ExecuteContext.UnitContext.State)
-	{
-		case EControlRigState::Init:
-		case EControlRigState::Update:
-		{
-			FVector UpVectorNormalized = UpVector;
-			UpVectorNormalized.Normalize();
+	FVector UpVectorNormalized = UpVector;
+	UpVectorNormalized.Normalize();
 
-			const float ClampedU = FMath::Clamp<float>(U, 0.f, 1.f);
-			const float ClampedRoll = FMath::Clamp<float>(Roll, -180.f, 180.f);
-		
-			FVector Tangent = Spline.TangentAtParam(ClampedU);
+	const float ClampedU = FMath::Clamp<float>(U, 0.f, 1.f);
+	const float ClampedRoll = FMath::Clamp<float>(Roll, -180.f, 180.f);
 
-			// Check if Tangent can be normalized. If not, keep the same tangent as before.
-			if (!Tangent.Normalize())
-			{
-				Tangent = Transform.ToMatrixNoScale().GetUnitAxis(EAxis::X);
-			}
-			FVector Binormal = FVector::CrossProduct(Tangent, UpVectorNormalized);
-			Binormal = Binormal.RotateAngleAxis(ClampedRoll * ClampedU, Tangent);
-			
-			FMatrix RotationMatrix = FRotationMatrix::MakeFromXZ(Tangent, Binormal);
-
-			Transform.SetFromMatrix(RotationMatrix);
-			Transform.SetTranslation(Spline.PositionAtParam(U));
-			break;
-		}
-		default:
-		{
-			checkNoEntry(); // Execute is only defined for Init and Update
-			break;
-		}
+	FVector Tangent = Spline.TangentAtParam(ClampedU);
+
+	// Check if Tangent can be normalized. If not, keep the same tangent as before.
+	if (!Tangent.Normalize())
+	{
+		Tangent = Transform.ToMatrixNoScale().GetUnitAxis(EAxis::X);
 	}
+	FVector Binormal = FVector::CrossProduct(Tangent, UpVectorNormalized);
+	Binormal = Binormal.RotateAngleAxis(ClampedRoll * ClampedU, Tangent);
+	
+	FMatrix RotationMatrix = FRotationMatrix::MakeFromXZ(Tangent, Binormal);
+
+	Transform.SetFromMatrix(RotationMatrix);
+	Transform.SetTranslation(Spline.PositionAtParam(U));
 }
 
 FRigUnit_TangentFromControlRigSpline_Execute()
@@ -178,42 +121,23 @@ FRigUnit_TangentFromControlRigSpline_Execute()
 	}
 #endif	
 	
-	switch (ExecuteContext.UnitContext.State)
-	{
-		case EControlRigState::Init:
-		case EControlRigState::Update:
-		{
-			const float ClampedU = FMath::Clamp<float>(U, 0.f, 1.f);
+	const float ClampedU = FMath::Clamp<float>(U, 0.f, 1.f);
 
-			FVector NewTangent = Spline.TangentAtParam(ClampedU);
+	FVector NewTangent = Spline.TangentAtParam(ClampedU);
 
-			// Check if Tangent can be normalized. If not, keep the same tangent as before.
-			if (!NewTangent.Normalize())
-			{
-				NewTangent = Tangent;
-			}
-			else
-			{
-				Tangent = NewTangent;
-			}
-			
-			break;
-		}
-		default:
-		{
-			checkNoEntry(); // Execute is only defined for Init and Update
-			break;
-		}
+	// Check if Tangent can be normalized. If not, keep the same tangent as before.
+	if (!NewTangent.Normalize())
+	{
+		NewTangent = Tangent;
+	}
+	else
+	{
+		Tangent = NewTangent;
 	}
 }
 
 FRigUnit_DrawControlRigSpline_Execute()
 {
-	if (ExecuteContext.UnitContext.State == EControlRigState::Init)
-	{
-		return;
-	}
-
 	if (ExecuteContext.UnitContext.DrawInterface == nullptr)
 	{
 		return;
@@ -265,20 +189,7 @@ FRigUnit_GetLengthControlRigSpline_Execute()
 	}
 #endif	
 	
-	switch (ExecuteContext.UnitContext.State)
-	{
-		case EControlRigState::Init:
-		case EControlRigState::Update:
-		{
-			Length = Spline.SplineData->AccumulatedLenth.Last();
-			break;
-		}
-		default:
-		{
-			checkNoEntry(); // Execute is only defined for Init and Update
-			break;
-		}
-	}
+	Length = Spline.SplineData->AccumulatedLenth.Last();
 }
 
 FRigUnit_FitChainToSplineCurve_Execute()
@@ -358,12 +269,6 @@ FRigUnit_FitChainToSplineCurveItemArray_Execute()
 	TArray<float>& ItemRotationT = WorkData.ItemRotationT;
 	TArray<FTransform>& ItemLocalTransforms = WorkData.ItemLocalTransforms;
 
-	if (ExecuteContext.UnitContext.State == EControlRigState::Init)
-	{
-		CachedItems.Reset();
-		return;
-	}
-
 	if(CachedItems.Num() == 0 && Items.Num() > 1)
 	{
 		CurvePositions.Reset();
@@ -776,11 +681,6 @@ FRigUnit_FitSplineCurveToChainItemArray_Execute()
 	}
 #endif	
 
-	if (ExecuteContext.UnitContext.State == EControlRigState::Init)
-	{
-		return;
-	}
-
 	if (Items.Num() < 4)
 	{
 		UE_CONTROLRIG_RIGUNIT_REPORT_WARNING(TEXT("Cannot create a spline with less than 4 points (%d received)."), Items.Num());
@@ -832,61 +732,47 @@ FRigUnit_ClosestParameterFromControlRigSpline_Execute()
 	}
 #endif	
 	
-	switch (ExecuteContext.UnitContext.State)
+	const int32 SampleCount = Spline.SplineData->SamplesArray.Num();
+
+	auto DistanceToSegment = [&](int32 Index0, int32 Index1, float& OutParam) -> float
 	{
-		case EControlRigState::Init:
-		case EControlRigState::Update:
-		{
-				const int32 SampleCount = Spline.SplineData->SamplesArray.Num();
+		const int32 MaxIndex = Spline.SplineData->SamplesArray.Num()-1;
+		const FVector& P0 = Spline.SplineData->SamplesArray[Index0];
+		const FVector& P1 = Spline.SplineData->SamplesArray[Index1];
 
-				auto DistanceToSegment = [&](int32 Index0, int32 Index1, float& OutParam) -> float
-				{
-					const int32 MaxIndex = Spline.SplineData->SamplesArray.Num()-1;
-					const FVector& P0 = Spline.SplineData->SamplesArray[Index0];
-					const FVector& P1 = Spline.SplineData->SamplesArray[Index1];
-
-					const FVector V = P1 - P0;
-					const FVector U = P0 - Position;
-					const float T = - U.Dot(V) / V.Dot(V);
-					if (T > 0 && T < 1)
-					{
-						OutParam = (1-T)*(Index0/(float)MaxIndex) + T*(Index1/(float)MaxIndex);
-						return U.Cross(V).Length() / V.Length();
-					}
+		const FVector V = P1 - P0;
+		const FVector U = P0 - Position;
+		const float T = - U.Dot(V) / V.Dot(V);
+		if (T > 0 && T < 1)
+		{
+			OutParam = (1-T)*(Index0/(float)MaxIndex) + T*(Index1/(float)MaxIndex);
+			return U.Cross(V).Length() / V.Length();
+		}
 
-					const float D0 = FVector::Distance(P0, Position);
-					const float D1 = FVector::Distance(P1, Position);
-					if (D0 < D1)
-					{
-						OutParam = Index0/(float)MaxIndex;
-						return D0;
-					}
+		const float D0 = FVector::Distance(P0, Position);
+		const float D1 = FVector::Distance(P1, Position);
+		if (D0 < D1)
+		{
+			OutParam = Index0/(float)MaxIndex;
+			return D0;
+		}
 
-					OutParam = Index1/(float)MaxIndex;
-					return D1;
-				};
+		OutParam = Index1/(float)MaxIndex;
+		return D1;
+	};
 
-				float ClosestDistance = TNumericLimits<float>::Max();
-				for (int32 i=1; i<SampleCount; ++i)
-				{
-					float Param;
-					float Distance = DistanceToSegment(i-1, i, Param);
+	float ClosestDistance = TNumericLimits<float>::Max();
+	for (int32 i=1; i<SampleCount; ++i)
+	{
+		float Param;
+		float Distance = DistanceToSegment(i-1, i, Param);
 
-					if (Distance < ClosestDistance)
-					{
-						ClosestDistance = Distance;
-						U = Param;
-					}
-				}				
-				
-			break;
-		}
-		default:
+		if (Distance < ClosestDistance)
 		{
-			checkNoEntry(); // Execute is only defined for Init and Update
-			break;
+			ClosestDistance = Distance;
+			U = Param;
 		}
-	}
+	}				
 }
 
 FRigUnit_ParameterAtPercentage_Execute()
@@ -904,48 +790,34 @@ FRigUnit_ParameterAtPercentage_Execute()
 	}
 #endif	
 	
-	switch (ExecuteContext.UnitContext.State)
-	{
-	case EControlRigState::Init:
-	case EControlRigState::Update:
-		{
-			const int32 SampleCount = Spline.SplineData->SamplesArray.Num();
+	const int32 SampleCount = Spline.SplineData->SamplesArray.Num();
 
-			float Length = Spline.SplineData->AccumulatedLenth.Last();
-			float ClampedPercentage = FMath::Clamp(Percentage, 0.f, 1.f);
+	float Length = Spline.SplineData->AccumulatedLenth.Last();
+	float ClampedPercentage = FMath::Clamp(Percentage, 0.f, 1.f);
 
-			float SearchLength = Length * ClampedPercentage;
-			int32 NextIndex = Algo::LowerBound(Spline.SplineData->AccumulatedLenth, SearchLength);
+	float SearchLength = Length * ClampedPercentage;
+	int32 NextIndex = Algo::LowerBound(Spline.SplineData->AccumulatedLenth, SearchLength);
 
-			if (NextIndex >= SampleCount)
-			{
-				U = 1.f;
-				return;
-			}
+	if (NextIndex >= SampleCount)
+	{
+		U = 1.f;
+		return;
+	}
 
-			if (NextIndex <= 0)
-			{
-				U = 0.f;
-				return;
-			}
+	if (NextIndex <= 0)
+	{
+		U = 0.f;
+		return;
+	}
 
-			float UPrev = (NextIndex - 1) / (float) (SampleCount - 1);
-			float UNext = NextIndex / (float) (SampleCount - 1);
+	float UPrev = (NextIndex - 1) / (float) (SampleCount - 1);
+	float UNext = NextIndex / (float) (SampleCount - 1);
 
-			float LengthPrev = Spline.SplineData->AccumulatedLenth[NextIndex-1];
-			float LengthNext = Spline.SplineData->AccumulatedLenth[NextIndex];
+	float LengthPrev = Spline.SplineData->AccumulatedLenth[NextIndex-1];
+	float LengthNext = Spline.SplineData->AccumulatedLenth[NextIndex];
 
-			float Interp = (SearchLength - LengthPrev) / (LengthNext - LengthPrev);
-			U = (1 - Interp) * UPrev + Interp * UNext;			
-				
-			break;
-		}
-	default:
-		{
-			checkNoEntry(); // Execute is only defined for Init and Update
-			break;
-		}
-	}
+	float Interp = (SearchLength - LengthPrev) / (LengthNext - LengthPrev);
+	U = (1 - Interp) * UPrev + Interp * UNext;			
 }
 
 
diff --git a/Engine/Plugins/Animation/RigLogic/Source/RigLogicModule/Private/RigUnit_RigLogic.cpp b/Engine/Plugins/Animation/RigLogic/Source/RigLogicModule/Private/RigUnit_RigLogic.cpp
--- a/Engine/Plugins/Animation/RigLogic/Source/RigLogicModule/Private/RigUnit_RigLogic.cpp
+++ b/Engine/Plugins/Animation/RigLogic/Source/RigLogicModule/Private/RigUnit_RigLogic.cpp
@@ -336,64 +336,54 @@ FRigUnit_RigLogic_Execute()
  	URigHierarchy* Hierarchy = ExecuteContext.Hierarchy;
 	if (Hierarchy)
 	{
-		switch (ExecuteContext.UnitContext.State)
+		if(!bIsInitialized)
 		{
-			case EControlRigState::Init:
+			//const double startTime = FPlatformTime::Seconds();
+			if (!Data.SkelMeshComponent.IsValid())
 			{
-				//const double startTime = FPlatformTime::Seconds();
-				if (!Data.SkelMeshComponent.IsValid())
-				{
-					Data.SkelMeshComponent = ExecuteContext.UnitContext.DataSourceRegistry->RequestSource<USkeletalMeshComponent>(UControlRig::OwnerComponent);
-					// In normal execution, Data.SkelMeshComponent will be nullptr at the beginning
-					// however, during unit testing we cannot fetch it from DataSourceRegistry 
-					// in that case, a mock version will be inserted into Data by unit test beforehand
-				}
-				if (!Data.SkelMeshComponent.IsValid() || Data.SkelMeshComponent->GetSkeletalMeshAsset() == nullptr)
-				{
-					return;
-				}
-				Data.CurrentLOD = Data.SkelMeshComponent->GetPredictedLODLevel();
-
-				// Fetch shared runtime context of rig from DNAAsset
-				TSharedPtr<FSharedRigRuntimeContext> RigRuntimeContext = GetSharedRigRuntimeContext(Data.SkelMeshComponent->GetSkeletalMeshAsset());
-				// Context is initialized with a BehaviorReader, which can be imported into SkeletalMesh from DNA file
-				// or overwritten by GeneSplicer when making a new character
-				Data.InitializeRigLogic(Hierarchy, RigRuntimeContext);
-				//const double delta = FPlatformTime::Seconds() - startTime;
-				//UE_LOG(LogRigLogicUnit, Warning, TEXT("RigLogic::Init execution time: %f"), delta);
-
-				break; 
+				Data.SkelMeshComponent = ExecuteContext.UnitContext.DataSourceRegistry->RequestSource<USkeletalMeshComponent>(UControlRig::OwnerComponent);
+				// In normal execution, Data.SkelMeshComponent will be nullptr at the beginning
+				// however, during unit testing we cannot fetch it from DataSourceRegistry 
+				// in that case, a mock version will be inserted into Data by unit test beforehand
 			}
-			case EControlRigState::Update:
+			if (!Data.SkelMeshComponent.IsValid() || Data.SkelMeshComponent->GetSkeletalMeshAsset() == nullptr)
 			{
-				//const double startTime = FPlatformTime::Seconds();
-				// Fetch shared runtime context of rig from DNAAsset
-				if (!Data.SkelMeshComponent.IsValid() || !Data.IsRigLogicInitialized() || Hierarchy == nullptr)
-				{
-					return;
-				}
-				Data.CurrentLOD = Data.SkelMeshComponent->GetPredictedLODLevel();
-				Data.CalculateRigLogic(Hierarchy);
+				return;
+			}
+			Data.CurrentLOD = Data.SkelMeshComponent->GetPredictedLODLevel();
+
+			// Fetch shared runtime context of rig from DNAAsset
+			TSharedPtr<FSharedRigRuntimeContext> RigRuntimeContext = GetSharedRigRuntimeContext(Data.SkelMeshComponent->GetSkeletalMeshAsset());
+			// Context is initialized with a BehaviorReader, which can be imported into SkeletalMesh from DNA file
+			// or overwritten by GeneSplicer when making a new character
+			Data.InitializeRigLogic(Hierarchy, RigRuntimeContext);
+			//const double delta = FPlatformTime::Seconds() - startTime;
+			//UE_LOG(LogRigLogicUnit, Warning, TEXT("RigLogic::Init execution time: %f"), delta);
+		}
+		
 
-				//Filing a struct so we can call the same method for updating joints from tests
-				TArrayView<const float> NeutralJointValues = Data.LocalRigRuntimeContext->RigLogic->GetRawNeutralJointValues();
-				TArrayView<const float> DeltaJointValues = Data.RigInstance->GetRawJointOutputs();
-				Data.UpdateJoints(Hierarchy, NeutralJointValues, DeltaJointValues);
+		//const double startTime = FPlatformTime::Seconds();
+		// Fetch shared runtime context of rig from DNAAsset
+		if (!Data.SkelMeshComponent.IsValid() || !Data.IsRigLogicInitialized() || Hierarchy == nullptr)
+		{
+			return;
+		}
+		Data.CurrentLOD = Data.SkelMeshComponent->GetPredictedLODLevel();
+		Data.CalculateRigLogic(Hierarchy);
 
-				TArrayView<const float> BlendShapeValues = Data.RigInstance->GetBlendShapeOutputs();
-				Data.UpdateBlendShapeCurves(Hierarchy, BlendShapeValues);
+		//Filing a struct so we can call the same method for updating joints from tests
+		TArrayView<const float> NeutralJointValues = Data.LocalRigRuntimeContext->RigLogic->GetRawNeutralJointValues();
+		TArrayView<const float> DeltaJointValues = Data.RigInstance->GetRawJointOutputs();
+		Data.UpdateJoints(Hierarchy, NeutralJointValues, DeltaJointValues);
 
-				TArrayView<const float> AnimMapOutputs = Data.RigInstance->GetAnimatedMapOutputs();
-				Data.UpdateAnimMapCurves(Hierarchy, AnimMapOutputs);
+		TArrayView<const float> BlendShapeValues = Data.RigInstance->GetBlendShapeOutputs();
+		Data.UpdateBlendShapeCurves(Hierarchy, BlendShapeValues);
 
-				//const double delta = FPlatformTime::Seconds() - startTime;
-				//UE_LOG(LogRigLogicUnit, Warning, TEXT("RigLogic::Update execution time: %f"), delta);
-			}
-			default:
-			{
-				break;
-			}
-		}
+		TArrayView<const float> AnimMapOutputs = Data.RigInstance->GetAnimatedMapOutputs();
+		Data.UpdateAnimMapCurves(Hierarchy, AnimMapOutputs);
+
+		//const double delta = FPlatformTime::Seconds() - startTime;
+		//UE_LOG(LogRigLogicUnit, Warning, TEXT("RigLogic::Update execution time: %f"), delta);
 	}
 }
 
diff --git a/Engine/Plugins/Animation/RigLogic/Source/RigLogicModule/Private/RigUnit_RigLogic_Test.cpp b/Engine/Plugins/Animation/RigLogic/Source/RigLogicModule/Private/RigUnit_RigLogic_Test.cpp
--- a/Engine/Plugins/Animation/RigLogic/Source/RigLogicModule/Private/RigUnit_RigLogic_Test.cpp
+++ b/Engine/Plugins/Animation/RigLogic/Source/RigLogicModule/Private/RigUnit_RigLogic_Test.cpp
@@ -277,7 +277,7 @@ IMPLEMENT_RIGUNIT_AUTOMATION_TEST(FRigUnit_RigLogic)
 	Unit.ExecuteContext.Hierarchy = Hierarchy;
 
 	//Test
-	InitAndExecute();
+	Execute();
 	TSharedPtr<FSharedRigRuntimeContext> SharedRigRuntimeContext = Test.GetSharedRigRuntimeContext(Test.GetData()->SkelMeshComponent->GetSkeletalMeshAsset());
 	//Assert
 	AddErrorIfFalse(
diff --git a/Engine/Plugins/Animation/RigLogic/Source/RigLogicModule/Public/RigUnit_RigLogic.h b/Engine/Plugins/Animation/RigLogic/Source/RigLogicModule/Public/RigUnit_RigLogic.h
--- a/Engine/Plugins/Animation/RigLogic/Source/RigLogicModule/Public/RigUnit_RigLogic.h
+++ b/Engine/Plugins/Animation/RigLogic/Source/RigLogicModule/Public/RigUnit_RigLogic.h
@@ -120,6 +120,11 @@ struct FRigUnit_RigLogic: public FRigUnitMutable
 {
 	GENERATED_BODY()
 
+	FRigUnit_RigLogic()
+		: Data()
+		, bIsInitialized(false)
+	{}
+
 public:
 #if WITH_DEV_AUTOMATION_TESTS
 	/** Allows accessing private Data property from unit tests **/
@@ -137,4 +142,8 @@ private:
 	// internal work data for the unit
 	UPROPERTY(transient)
 	FRigUnit_RigLogic_Data Data;
+
+	// flag to determine if the node has been initialized
+	UPROPERTY(transient)
+	bool bIsInitialized;
 };
diff --git a/Engine/Plugins/Experimental/AnimNext/Source/AnimNextInterfaceGraph/Private/AnimNextInterfaceGraph.cpp b/Engine/Plugins/Experimental/AnimNext/Source/AnimNextInterfaceGraph/Private/AnimNextInterfaceGraph.cpp
--- a/Engine/Plugins/Experimental/AnimNext/Source/AnimNextInterfaceGraph/Private/AnimNextInterfaceGraph.cpp
+++ b/Engine/Plugins/Experimental/AnimNext/Source/AnimNextInterfaceGraph/Private/AnimNextInterfaceGraph.cpp
@@ -32,7 +32,6 @@ bool UAnimNextInterfaceGraph::GetDataImpl(const UE::AnimNext::Interface::FContex
 	if(RigVM)
 	{
 		FRigUnitContext RigUnitContext;
-		RigUnitContext.State = EControlRigState::Update;
 		FAnimNextInterfaceUnitContext AnimNextInterfaceUnitContext(this, Context, bResult);
 		bResult &=
 			(RigVM->Execute(TArray<URigVMMemoryStorage*>(), FRigUnit_AnimNextInterfaceBeginExecution::EventName) != ERigVMExecuteResult::Failed);
diff --git a/Engine/Plugins/Experimental/AnimNext/Source/AnimNextInterfaceGraphUncookedOnly/Private/AnimNextInterfaceGraphUncookedOnlyUtils.cpp b/Engine/Plugins/Experimental/AnimNext/Source/AnimNextInterfaceGraphUncookedOnly/Private/AnimNextInterfaceGraphUncookedOnlyUtils.cpp
--- a/Engine/Plugins/Experimental/AnimNext/Source/AnimNextInterfaceGraphUncookedOnly/Private/AnimNextInterfaceGraphUncookedOnlyUtils.cpp
+++ b/Engine/Plugins/Experimental/AnimNext/Source/AnimNextInterfaceGraphUncookedOnly/Private/AnimNextInterfaceGraphUncookedOnlyUtils.cpp
@@ -36,15 +36,6 @@ void FUtils::Compile(UAnimNextInterfaceGraph* InGraph)
 
 	InGraph->VMRuntimeSettings = EditorData->VMRuntimeSettings;
 
-	FRigUnitContext InitContext;
-	InitContext.State = EControlRigState::Init;
-
-	FRigUnitContext UpdateContext = InitContext;
-	UpdateContext.State = EControlRigState::Update;
-
-	void* InitContextPtr = &InitContext;
-	void* UpdateContextPtr = &UpdateContext;
-
 	EditorData->CompileLog.Messages.Reset();
 	EditorData->CompileLog.NumErrors = EditorData->CompileLog.NumWarnings = 0;
 
diff --git a/Engine/Plugins/Experimental/FullBodyIK/Source/FullBodyIK/Private/RigUnit_FullbodyIK.cpp b/Engine/Plugins/Experimental/FullBodyIK/Source/FullBodyIK/Private/RigUnit_FullbodyIK.cpp
--- a/Engine/Plugins/Experimental/FullBodyIK/Source/FullBodyIK/Private/RigUnit_FullbodyIK.cpp
+++ b/Engine/Plugins/Experimental/FullBodyIK/Source/FullBodyIK/Private/RigUnit_FullbodyIK.cpp
@@ -215,337 +215,334 @@ FRigUnit_FullbodyIK_Execute()
 		AddEffectors(Hierarchy, Root, EffectorsView, LinkData, EffectorTargets, EffectorLinkIndices, LinkDataToHierarchyIndices, HierarchyToLinkDataMap, SolverProperty);
 	}
 
-	if (ExecuteContext.UnitContext.State == EControlRigState::Update)
-	{
-		DECLARE_SCOPE_HIERARCHICAL_COUNTER(TEXT("Update"))
+	DECLARE_SCOPE_HIERARCHICAL_COUNTER(TEXT("Update"))
 
-		if (LinkDataToHierarchyIndices.Num() > 0)
+	if (LinkDataToHierarchyIndices.Num() > 0)
+	{
+		// we do this every frame for now
+		if (Constraints.Num() > 0)
 		{
-			// we do this every frame for now
-			if (Constraints.Num() > 0)
-			{
-				DECLARE_SCOPE_HIERARCHICAL_COUNTER(TEXT("Build Constraint"))
-				//Build constraints
-				const TArrayView<const FFBIKConstraintOption> ConstraintsView(Constraints.GetData(), Constraints.Num());
-				FBIKConstraintLib::BuildConstraints(ConstraintsView, InternalConstraints, Hierarchy, LinkData, LinkDataToHierarchyIndices, HierarchyToLinkDataMap);
-			}
+			DECLARE_SCOPE_HIERARCHICAL_COUNTER(TEXT("Build Constraint"))
+			//Build constraints
+			const TArrayView<const FFBIKConstraintOption> ConstraintsView(Constraints.GetData(), Constraints.Num());
+			FBIKConstraintLib::BuildConstraints(ConstraintsView, InternalConstraints, Hierarchy, LinkData, LinkDataToHierarchyIndices, HierarchyToLinkDataMap);
+		}
 
-			// during only editor and update
-			// we expect solver type changes, it will reinit
-			// InternalConstraints can't be changed during runtime
-			if (InternalConstraints.Num() > 0)
-			{
-				WorkData.IKSolver.SetPostProcessDelegateForIteration(FPostProcessDelegateForIteration::CreateStatic(&FBIKConstraintLib::ApplyConstraint, &InternalConstraints));
-			}
-			else
-			{
-				WorkData.IKSolver.ClearPostProcessDelegateForIteration();
-			}
-			// first set mid effector's transform
+		// during only editor and update
+		// we expect solver type changes, it will reinit
+		// InternalConstraints can't be changed during runtime
+		if (InternalConstraints.Num() > 0)
+		{
+			WorkData.IKSolver.SetPostProcessDelegateForIteration(FPostProcessDelegateForIteration::CreateStatic(&FBIKConstraintLib::ApplyConstraint, &InternalConstraints));
+		}
+		else
+		{
+			WorkData.IKSolver.ClearPostProcessDelegateForIteration();
+		}
+		// first set mid effector's transform
 
-			// before update we finalize motion scale
-			// this code may go away once we have constraint
-			// update link data and end effectors
-			for (int32 LinkIndex = 0; LinkIndex < LinkData.Num(); ++LinkIndex)
-			{
-				const FRigElementKey& Item = *LinkDataToHierarchyIndices.Find(LinkIndex);
-				LinkData[LinkIndex].SetTransform(Hierarchy->GetGlobalTransform(Item));
-				// @todo: fix this somewhere else - we can add this to prepare step
-				// @todo: we update motion scale here, then?
-				LinkData[LinkIndex].FinalizeForSolver();
-			}
+		// before update we finalize motion scale
+		// this code may go away once we have constraint
+		// update link data and end effectors
+		for (int32 LinkIndex = 0; LinkIndex < LinkData.Num(); ++LinkIndex)
+		{
+			const FRigElementKey& Item = *LinkDataToHierarchyIndices.Find(LinkIndex);
+			LinkData[LinkIndex].SetTransform(Hierarchy->GetGlobalTransform(Item));
+			// @todo: fix this somewhere else - we can add this to prepare step
+			// @todo: we update motion scale here, then?
+			LinkData[LinkIndex].FinalizeForSolver();
+		}
 
-			// update mid effector info
+		// update mid effector info
 
-			const float LinearMotionStrength = FMath::Max(SolverProperty.LinearMotionStrength, SolverProperty.MinLinearMotionStrength);
-			const float AngularMotionStrength = FMath::Max(SolverProperty.AngularMotionStrength, SolverProperty.MinAngularMotionStrength);
-			const float LinearRange = LinearMotionStrength - SolverProperty.MinLinearMotionStrength;
-			const float AngularRange = AngularMotionStrength - SolverProperty.MinAngularMotionStrength;
+		const float LinearMotionStrength = FMath::Max(SolverProperty.LinearMotionStrength, SolverProperty.MinLinearMotionStrength);
+		const float AngularMotionStrength = FMath::Max(SolverProperty.AngularMotionStrength, SolverProperty.MinAngularMotionStrength);
+		const float LinearRange = LinearMotionStrength - SolverProperty.MinLinearMotionStrength;
+		const float AngularRange = AngularMotionStrength - SolverProperty.MinAngularMotionStrength;
 
-			// update end effector info
-			for (int32 EffectorIndex = 0; EffectorIndex < Effectors.Num(); ++EffectorIndex)
+		// update end effector info
+		for (int32 EffectorIndex = 0; EffectorIndex < Effectors.Num(); ++EffectorIndex)
+		{
+			int32 EffectorLinkIndex = EffectorLinkIndices[EffectorIndex];
+			if (EffectorLinkIndex != INDEX_NONE)
 			{
-				int32 EffectorLinkIndex = EffectorLinkIndices[EffectorIndex];
-				if (EffectorLinkIndex != INDEX_NONE)
+				FFBIKEffectorTarget* EffectorTarget = EffectorTargets.Find(EffectorLinkIndex);
+				if (EffectorTarget)
 				{
-					FFBIKEffectorTarget* EffectorTarget = EffectorTargets.Find(EffectorLinkIndex);
-					if (EffectorTarget)
-					{
-						const FFBIKEndEffector& CurEffector = Effectors[EffectorIndex];
-						const FVector CurrentLinkLocation = LinkData[EffectorLinkIndex].GetTransform().GetLocation();
-						const FQuat CurrentLinkRotation = LinkData[EffectorLinkIndex].GetTransform().GetRotation();
-						const FVector& EffectorLocation = CurEffector.Position;
-						const FQuat& EffectorRotation = CurEffector.Rotation;
-						EffectorTarget->Position = FMath::Lerp(CurrentLinkLocation, EffectorLocation, CurEffector.PositionAlpha);
-						EffectorTarget->Rotation = FMath::Lerp(CurrentLinkRotation, EffectorRotation, CurEffector.RotationAlpha);
-						EffectorTarget->InitialPositionDistance = (EffectorLocation - CurrentLinkLocation).Size();
-						EffectorTarget->InitialRotationDistance = (FBIKUtil::GetScaledRotationAxis(EffectorRotation) - FBIKUtil::GetScaledRotationAxis(CurrentLinkRotation)).Size();
-
-						const float Pull = FMath::Clamp(CurEffector.Pull, 0.f, 1.f);
-						// we want some impact of Pull, in order for Pull to have some impact, we clamp to some number
-						const float TargetClamp = FMath::Clamp(SolverProperty.DefaultTargetClamp, 0.f, 0.7f);
-						const float Scale = TargetClamp + Pull * (1.f - TargetClamp);
-						// Pull set up
-						EffectorTarget->LinearMotionStrength = LinearRange * Scale + SolverProperty.MinLinearMotionStrength;
-						EffectorTarget->AngularMotionStrength = AngularRange * Scale + SolverProperty.MinAngularMotionStrength;
-						EffectorTarget->ConvergeScale = Scale;
-						EffectorTarget->TargetClampScale = Scale;
-
-						EffectorTarget->bPositionEnabled = true;
-						EffectorTarget->bRotationEnabled = true;
-					}
+					const FFBIKEndEffector& CurEffector = Effectors[EffectorIndex];
+					const FVector CurrentLinkLocation = LinkData[EffectorLinkIndex].GetTransform().GetLocation();
+					const FQuat CurrentLinkRotation = LinkData[EffectorLinkIndex].GetTransform().GetRotation();
+					const FVector& EffectorLocation = CurEffector.Position;
+					const FQuat& EffectorRotation = CurEffector.Rotation;
+					EffectorTarget->Position = FMath::Lerp(CurrentLinkLocation, EffectorLocation, CurEffector.PositionAlpha);
+					EffectorTarget->Rotation = FMath::Lerp(CurrentLinkRotation, EffectorRotation, CurEffector.RotationAlpha);
+					EffectorTarget->InitialPositionDistance = (EffectorLocation - CurrentLinkLocation).Size();
+					EffectorTarget->InitialRotationDistance = (FBIKUtil::GetScaledRotationAxis(EffectorRotation) - FBIKUtil::GetScaledRotationAxis(CurrentLinkRotation)).Size();
+
+					const float Pull = FMath::Clamp(CurEffector.Pull, 0.f, 1.f);
+					// we want some impact of Pull, in order for Pull to have some impact, we clamp to some number
+					const float TargetClamp = FMath::Clamp(SolverProperty.DefaultTargetClamp, 0.f, 0.7f);
+					const float Scale = TargetClamp + Pull * (1.f - TargetClamp);
+					// Pull set up
+					EffectorTarget->LinearMotionStrength = LinearRange * Scale + SolverProperty.MinLinearMotionStrength;
+					EffectorTarget->AngularMotionStrength = AngularRange * Scale + SolverProperty.MinAngularMotionStrength;
+					EffectorTarget->ConvergeScale = Scale;
+					EffectorTarget->TargetClampScale = Scale;
+
+					EffectorTarget->bPositionEnabled = true;
+					EffectorTarget->bRotationEnabled = true;
 				}
 			}
+		}
 
-			TArray<FJacobianDebugData>& DebugData = WorkData.DebugData;
-			DebugData.Reset();
+		TArray<FJacobianDebugData>& DebugData = WorkData.DebugData;
+		DebugData.Reset();
 
-			const bool bDebugEnabled = DebugOption.bDrawDebugHierarchy || DebugOption.bDrawDebugEffector || DebugOption.bDrawDebugConstraints;
+		const bool bDebugEnabled = DebugOption.bDrawDebugHierarchy || DebugOption.bDrawDebugEffector || DebugOption.bDrawDebugConstraints;
 
-			// we can't reuse memory until we fix the memory issue on RigVM
-			{
-				DECLARE_SCOPE_HIERARCHICAL_COUNTER(TEXT("Solver"))
-				FJacobianSolver_FullbodyIK& IKSolver = WorkData.IKSolver;
-				IKSolver.SolveJacobianIK(LinkData, EffectorTargets,
-					JacobianIK::FSolverParameter(SolverProperty.Damping, true, false, (SolverProperty.bUseJacobianTranspose) ? EJacobianSolver::JacobianTranspose : EJacobianSolver::JacobianPIDLS),
-					SolverProperty.MaxIterations, SolverProperty.Precision, (bDebugEnabled) ? &DebugData : nullptr);
+		// we can't reuse memory until we fix the memory issue on RigVM
+		{
+			DECLARE_SCOPE_HIERARCHICAL_COUNTER(TEXT("Solver"))
+			FJacobianSolver_FullbodyIK& IKSolver = WorkData.IKSolver;
+			IKSolver.SolveJacobianIK(LinkData, EffectorTargets,
+				JacobianIK::FSolverParameter(SolverProperty.Damping, true, false, (SolverProperty.bUseJacobianTranspose) ? EJacobianSolver::JacobianTranspose : EJacobianSolver::JacobianPIDLS),
+				SolverProperty.MaxIterations, SolverProperty.Precision, (bDebugEnabled) ? &DebugData : nullptr);
 
-				if (MotionProperty.bForceEffectorRotationTarget)
+			if (MotionProperty.bForceEffectorRotationTarget)
+			{
+				// if position is reached, we force rotation target
+				for (int32 EffectorIndex = 0; EffectorIndex < Effectors.Num(); ++EffectorIndex)
 				{
-					// if position is reached, we force rotation target
-					for (int32 EffectorIndex = 0; EffectorIndex < Effectors.Num(); ++EffectorIndex)
+					int32 EffectorLinkIndex = EffectorLinkIndices[EffectorIndex];
+					if (EffectorLinkIndex != INDEX_NONE)
 					{
-						int32 EffectorLinkIndex = EffectorLinkIndices[EffectorIndex];
-						if (EffectorLinkIndex != INDEX_NONE)
+						FFBIKEffectorTarget* EffectorTarget = EffectorTargets.Find(EffectorLinkIndex);
+						if (EffectorTarget && EffectorTarget->bRotationEnabled)
 						{
-							FFBIKEffectorTarget* EffectorTarget = EffectorTargets.Find(EffectorLinkIndex);
-							if (EffectorTarget && EffectorTarget->bRotationEnabled)
-							{
-								bool bApplyRotation = true;
+							bool bApplyRotation = true;
 
-								if (MotionProperty.bOnlyApplyWhenReachedToTarget)
-								{
-									// only do this when position is reached? This will conflict with converge scale
-									const FVector& BonePosition = LinkData[EffectorLinkIndex].GetTransform().GetLocation();
-									const FVector& TargetPosition = EffectorTarget->Position;
+							if (MotionProperty.bOnlyApplyWhenReachedToTarget)
+							{
+								// only do this when position is reached? This will conflict with converge scale
+								const FVector& BonePosition = LinkData[EffectorLinkIndex].GetTransform().GetLocation();
+								const FVector& TargetPosition = EffectorTarget->Position;
 
-									bApplyRotation = (FVector(BonePosition-TargetPosition).SizeSquared() <= SolverProperty.Precision * SolverProperty.Precision);
-								}
+								bApplyRotation = (FVector(BonePosition-TargetPosition).SizeSquared() <= SolverProperty.Precision * SolverProperty.Precision);
+							}
 
-								if (bApplyRotation)
-								{
-									FQuat NewRotation = EffectorTarget->Rotation;
-									FTransform NewTransform = LinkData[EffectorLinkIndex].GetTransform();
-									NewTransform.SetRotation(NewRotation);
-									LinkData[EffectorLinkIndex].SetTransform(NewTransform);
-								}
+							if (bApplyRotation)
+							{
+								FQuat NewRotation = EffectorTarget->Rotation;
+								FTransform NewTransform = LinkData[EffectorLinkIndex].GetTransform();
+								NewTransform.SetRotation(NewRotation);
+								LinkData[EffectorLinkIndex].SetTransform(NewTransform);
 							}
 						}
 					}
 				}
 			}
+		}
 
-			///////////////////////////////////////////////////////////////////////////
-			// debug draw start
-			///////////////////////////////////////////////////////////////////////////
-			if (bDebugEnabled && ExecuteContext.UnitContext.DrawInterface != nullptr)
+		///////////////////////////////////////////////////////////////////////////
+		// debug draw start
+		///////////////////////////////////////////////////////////////////////////
+		if (bDebugEnabled && ExecuteContext.UnitContext.DrawInterface != nullptr)
+		{
+			const int32 DebugDataNum = DebugData.Num();
+			if (DebugData.Num() > 0)
 			{
-				const int32 DebugDataNum = DebugData.Num();
-				if (DebugData.Num() > 0)
+				for (int32 DebugIndex = DebugDataNum - 1; DebugIndex >= 0; --DebugIndex)
 				{
-					for (int32 DebugIndex = DebugDataNum - 1; DebugIndex >= 0; --DebugIndex)
-					{
-						const TArray<FFBIKLinkData>& LocalLink = DebugData[DebugIndex].LinkData;
+					const TArray<FFBIKLinkData>& LocalLink = DebugData[DebugIndex].LinkData;
 
-						FTransform Offset = DebugOption.DrawWorldOffset;
-						Offset.SetLocation(Offset.GetLocation() * (DebugDataNum - DebugIndex));
+					FTransform Offset = DebugOption.DrawWorldOffset;
+					Offset.SetLocation(Offset.GetLocation() * (DebugDataNum - DebugIndex));
 
-						if (DebugOption.bDrawDebugHierarchy)
+					if (DebugOption.bDrawDebugHierarchy)
+					{
+						for (int32 LinkIndex = 0; LinkIndex < LocalLink.Num(); ++LinkIndex)
 						{
-							for (int32 LinkIndex = 0; LinkIndex < LocalLink.Num(); ++LinkIndex)
-							{
-								const FFBIKLinkData& Data = LocalLink[LinkIndex];
+							const FFBIKLinkData& Data = LocalLink[LinkIndex];
 
-								FLinearColor DrawColor = FLinearColor::White;
+							FLinearColor DrawColor = FLinearColor::White;
 
-								float LineThickness = 0.f;
-								if (DebugOption.bColorAngularMotionStrength || DebugOption.bColorLinearMotionStrength)
+							float LineThickness = 0.f;
+							if (DebugOption.bColorAngularMotionStrength || DebugOption.bColorLinearMotionStrength)
+							{
+								DrawColor = FLinearColor::Black;
+								if (DebugOption.bColorAngularMotionStrength)
 								{
-									DrawColor = FLinearColor::Black;
-									if (DebugOption.bColorAngularMotionStrength)
+									const float Range = FMath::Max(SolverProperty.AngularMotionStrength - SolverProperty.MinAngularMotionStrength, 0.f);
+									if (Range > 0.f)
 									{
-										const float Range = FMath::Max(SolverProperty.AngularMotionStrength - SolverProperty.MinAngularMotionStrength, 0.f);
-										if (Range > 0.f)
-										{
-											float CurrentStrength = Data.GetAngularMotionStrength() - SolverProperty.MinAngularMotionStrength;
-											float Alpha = FMath::Clamp(CurrentStrength / Range, 0.f, 1.f);
-											DrawColor.R = LineThickness = Alpha;
-										}
+										float CurrentStrength = Data.GetAngularMotionStrength() - SolverProperty.MinAngularMotionStrength;
+										float Alpha = FMath::Clamp(CurrentStrength / Range, 0.f, 1.f);
+										DrawColor.R = LineThickness = Alpha;
 									}
-									else if (DebugOption.bColorLinearMotionStrength)
+								}
+								else if (DebugOption.bColorLinearMotionStrength)
+								{
+									const float Range = FMath::Max(SolverProperty.LinearMotionStrength - SolverProperty.MinLinearMotionStrength, 0.f);
+									if (Range > 0.f)
 									{
-										const float Range = FMath::Max(SolverProperty.LinearMotionStrength - SolverProperty.MinLinearMotionStrength, 0.f);
-										if (Range > 0.f)
-										{
-											float CurrentStrength = Data.GetLinearMotionStrength() - SolverProperty.MinLinearMotionStrength;
-											float Alpha = FMath::Clamp(CurrentStrength / Range, 0.f, 1.f);
-											DrawColor.B = LineThickness = Alpha;
-										}
+										float CurrentStrength = Data.GetLinearMotionStrength() - SolverProperty.MinLinearMotionStrength;
+										float Alpha = FMath::Clamp(CurrentStrength / Range, 0.f, 1.f);
+										DrawColor.B = LineThickness = Alpha;
 									}
 								}
+							}
 
-								if (Data.ParentLinkIndex != INDEX_NONE)
-								{
-									const FFBIKLinkData& ParentData = LocalLink[Data.ParentLinkIndex];
-									ExecuteContext.UnitContext.DrawInterface->DrawLine(Offset, Data.GetPreviousTransform().GetLocation(), ParentData.GetPreviousTransform().GetLocation(), DrawColor, LineThickness);
-								}
+							if (Data.ParentLinkIndex != INDEX_NONE)
+							{
+								const FFBIKLinkData& ParentData = LocalLink[Data.ParentLinkIndex];
+								ExecuteContext.UnitContext.DrawInterface->DrawLine(Offset, Data.GetPreviousTransform().GetLocation(), ParentData.GetPreviousTransform().GetLocation(), DrawColor, LineThickness);
+							}
 
-								if (DebugOption.bDrawDebugAxes)
-								{
-									ExecuteContext.UnitContext.DrawInterface->DrawAxes(Offset, Data.GetPreviousTransform(), DebugOption.DrawSize);
-								}
+							if (DebugOption.bDrawDebugAxes)
+							{
+								ExecuteContext.UnitContext.DrawInterface->DrawAxes(Offset, Data.GetPreviousTransform(), DebugOption.DrawSize);
 							}
 						}
+					}
 
-						if (DebugOption.bDrawDebugEffector)
+					if (DebugOption.bDrawDebugEffector)
+					{
+						for (auto Iter = EffectorTargets.CreateConstIterator(); Iter; ++Iter)
 						{
-							for (auto Iter = EffectorTargets.CreateConstIterator(); Iter; ++Iter)
+							const FFBIKEffectorTarget& EffectorTarget = Iter.Value();
+							if (EffectorTarget.bPositionEnabled)
 							{
-								const FFBIKEffectorTarget& EffectorTarget = Iter.Value();
-								if (EffectorTarget.bPositionEnabled)
-								{
-									// draw effector target locations
-									ExecuteContext.UnitContext.DrawInterface->DrawBox(Offset, FTransform(EffectorTarget.Position), FLinearColor::Yellow, DebugOption.DrawSize);
-								}
-
-								// draw effector link location
-								ExecuteContext.UnitContext.DrawInterface->DrawBox(Offset, LocalLink[Iter.Key()].GetPreviousTransform(), FLinearColor::Green, DebugOption.DrawSize);
+								// draw effector target locations
+								ExecuteContext.UnitContext.DrawInterface->DrawBox(Offset, FTransform(EffectorTarget.Position), FLinearColor::Yellow, DebugOption.DrawSize);
 							}
 
-							for (int32 Index = 0; Index < DebugData[DebugIndex].TargetVectorSources.Num(); ++Index)
-							{
-								// draw arrow to the target
-								ExecuteContext.UnitContext.DrawInterface->DrawLine(Offset, DebugData[DebugIndex].TargetVectorSources[Index].GetLocation(), 
-									DebugData[DebugIndex].TargetVectorSources[Index].GetLocation() + DebugData[DebugIndex].TargetVectors[Index], FLinearColor::Red);
-							}
+							// draw effector link location
+							ExecuteContext.UnitContext.DrawInterface->DrawBox(Offset, LocalLink[Iter.Key()].GetPreviousTransform(), FLinearColor::Green, DebugOption.DrawSize);
+						}
+
+						for (int32 Index = 0; Index < DebugData[DebugIndex].TargetVectorSources.Num(); ++Index)
+						{
+							// draw arrow to the target
+							ExecuteContext.UnitContext.DrawInterface->DrawLine(Offset, DebugData[DebugIndex].TargetVectorSources[Index].GetLocation(), 
+								DebugData[DebugIndex].TargetVectorSources[Index].GetLocation() + DebugData[DebugIndex].TargetVectors[Index], FLinearColor::Red);
 						}
 					}
 				}
+			}
 
-				if (DebugOption.bDrawDebugConstraints && InternalConstraints.Num())
-				{
-					FTransform Offset = FTransform::Identity;
+			if (DebugOption.bDrawDebugConstraints && InternalConstraints.Num())
+			{
+				FTransform Offset = FTransform::Identity;
 
-					// draw frame if active
-					for (int32 Index = 0; Index < Constraints.Num(); ++Index)
+				// draw frame if active
+				for (int32 Index = 0; Index < Constraints.Num(); ++Index)
+				{
+					if (Constraints[Index].bEnabled)
 					{
-						if (Constraints[Index].bEnabled)
+						if (Constraints[Index].Item.IsValid())
 						{
-							if (Constraints[Index].Item.IsValid())
+							const int32* Found = HierarchyToLinkDataMap.Find(Constraints[Index].Item);
+							if (Found)
 							{
-								const int32* Found = HierarchyToLinkDataMap.Find(Constraints[Index].Item);
-								if (Found)
-								{
-									FTransform ConstraintFrame = LinkData[*Found].GetTransform();
-									ConstraintFrame.ConcatenateRotation(FQuat(Constraints[Index].OffsetRotation));
-									ExecuteContext.UnitContext.DrawInterface->DrawAxes(Offset, ConstraintFrame, 2.f);
-								}
+								FTransform ConstraintFrame = LinkData[*Found].GetTransform();
+								ConstraintFrame.ConcatenateRotation(FQuat(Constraints[Index].OffsetRotation));
+								ExecuteContext.UnitContext.DrawInterface->DrawAxes(Offset, ConstraintFrame, 2.f);
 							}
 						}
 					}
+				}
 
-					for (int32 Index = 0; Index < InternalConstraints.Num(); ++Index)
+				for (int32 Index = 0; Index < InternalConstraints.Num(); ++Index)
+				{
+					// for now we have rotation limit only
+					if (InternalConstraints[Index].IsType< FRotationLimitConstraint>())
 					{
-						// for now we have rotation limit only
-						if (InternalConstraints[Index].IsType< FRotationLimitConstraint>())
+						FRotationLimitConstraint& LimitConstraint = InternalConstraints[Index].Get<FRotationLimitConstraint>();
+						const FQuat LocalRefRotation = LimitConstraint.RelativelRefPose.GetRotation();
+						FTransform RotationTransform = LinkData[LimitConstraint.ConstrainedIndex].GetTransform();
+						// base is parent transform but in their space, we can get there by inversing local ref rotation
+						FTransform BaseTransform = FTransform(LocalRefRotation).GetRelativeTransformReverse(RotationTransform);
+						BaseTransform.ConcatenateRotation(LimitConstraint.BaseFrameOffset);
+						ExecuteContext.UnitContext.DrawInterface->DrawAxes(Offset, BaseTransform, 5.f, 1.f);
+
+						// current transform
+						const FQuat LocalRotation = BaseTransform.GetRotation().Inverse() * RotationTransform.GetRotation();
+						const FQuat DeltaTransform = LocalRefRotation.Inverse() * LocalRotation;
+						RotationTransform.SetRotation(BaseTransform.GetRotation() * DeltaTransform);
+						RotationTransform.NormalizeRotation();
+						RotationTransform.SetLocation(BaseTransform.GetLocation());
+
+						// draw ref pose on their current transform
+						ExecuteContext.UnitContext.DrawInterface->DrawAxes(Offset, RotationTransform, 10.f, 1.f);
+
+						FVector XAxis = BaseTransform.GetUnitAxis(EAxis::X);
+						FVector YAxis = BaseTransform.GetUnitAxis(EAxis::Y);
+						FVector ZAxis = BaseTransform.GetUnitAxis(EAxis::Z);
+
+						if (LimitConstraint.bXLimitSet)
 						{
-							FRotationLimitConstraint& LimitConstraint = InternalConstraints[Index].Get<FRotationLimitConstraint>();
-							const FQuat LocalRefRotation = LimitConstraint.RelativelRefPose.GetRotation();
-							FTransform RotationTransform = LinkData[LimitConstraint.ConstrainedIndex].GetTransform();
-							// base is parent transform but in their space, we can get there by inversing local ref rotation
-							FTransform BaseTransform = FTransform(LocalRefRotation).GetRelativeTransformReverse(RotationTransform);
-							BaseTransform.ConcatenateRotation(LimitConstraint.BaseFrameOffset);
-							ExecuteContext.UnitContext.DrawInterface->DrawAxes(Offset, BaseTransform, 5.f, 1.f);
-
-							// current transform
-							const FQuat LocalRotation = BaseTransform.GetRotation().Inverse() * RotationTransform.GetRotation();
-							const FQuat DeltaTransform = LocalRefRotation.Inverse() * LocalRotation;
-							RotationTransform.SetRotation(BaseTransform.GetRotation() * DeltaTransform);
-							RotationTransform.NormalizeRotation();
-							RotationTransform.SetLocation(BaseTransform.GetLocation());
-
-							// draw ref pose on their current transform
-							ExecuteContext.UnitContext.DrawInterface->DrawAxes(Offset, RotationTransform, 10.f, 1.f);
-
-							FVector XAxis = BaseTransform.GetUnitAxis(EAxis::X);
-							FVector YAxis = BaseTransform.GetUnitAxis(EAxis::Y);
-							FVector ZAxis = BaseTransform.GetUnitAxis(EAxis::Z);
-
-							if (LimitConstraint.bXLimitSet)
-							{
-								FTransform XAxisConeTM(YAxis, XAxis ^ YAxis, XAxis, BaseTransform.GetTranslation());
-								XAxisConeTM.SetRotation(FQuat(XAxis, 0.f) * XAxisConeTM.GetRotation());
-								XAxisConeTM.SetScale3D(FVector(30.f));
-								ExecuteContext.UnitContext.DrawInterface->DrawCone(Offset, XAxisConeTM, LimitConstraint.Limit.X, 0.0f, 24, false, FLinearColor::Red, GEngine->ConstraintLimitMaterialX->GetRenderProxy());
-							}
-
-							if (LimitConstraint.bYLimitSet)
-							{
-								FTransform YAxisConeTM(ZAxis, YAxis ^ ZAxis, YAxis, BaseTransform.GetTranslation());
-								YAxisConeTM.SetRotation(FQuat(YAxis, 0.f) * YAxisConeTM.GetRotation());
-								YAxisConeTM.SetScale3D(FVector(30.f));
-								ExecuteContext.UnitContext.DrawInterface->DrawCone(Offset, YAxisConeTM, LimitConstraint.Limit.Y, 0.0f, 24, false, FLinearColor::Green, GEngine->ConstraintLimitMaterialY->GetRenderProxy());
-							}
+							FTransform XAxisConeTM(YAxis, XAxis ^ YAxis, XAxis, BaseTransform.GetTranslation());
+							XAxisConeTM.SetRotation(FQuat(XAxis, 0.f) * XAxisConeTM.GetRotation());
+							XAxisConeTM.SetScale3D(FVector(30.f));
+							ExecuteContext.UnitContext.DrawInterface->DrawCone(Offset, XAxisConeTM, LimitConstraint.Limit.X, 0.0f, 24, false, FLinearColor::Red, GEngine->ConstraintLimitMaterialX->GetRenderProxy());
+						}
 
-							if (LimitConstraint.bZLimitSet)
-							{
-								FTransform ZAxisConeTM(XAxis, ZAxis ^ XAxis, ZAxis, BaseTransform.GetTranslation());
-								ZAxisConeTM.SetRotation(FQuat(ZAxis, 0.f) * ZAxisConeTM.GetRotation());
-								ZAxisConeTM.SetScale3D(FVector(30.f));
-								ExecuteContext.UnitContext.DrawInterface->DrawCone(Offset, ZAxisConeTM, LimitConstraint.Limit.Z, 0.0f, 24, false, FLinearColor::Blue, GEngine->ConstraintLimitMaterialZ->GetRenderProxy());
-							}
+						if (LimitConstraint.bYLimitSet)
+						{
+							FTransform YAxisConeTM(ZAxis, YAxis ^ ZAxis, YAxis, BaseTransform.GetTranslation());
+							YAxisConeTM.SetRotation(FQuat(YAxis, 0.f) * YAxisConeTM.GetRotation());
+							YAxisConeTM.SetScale3D(FVector(30.f));
+							ExecuteContext.UnitContext.DrawInterface->DrawCone(Offset, YAxisConeTM, LimitConstraint.Limit.Y, 0.0f, 24, false, FLinearColor::Green, GEngine->ConstraintLimitMaterialY->GetRenderProxy());
 						}
 
-						if (InternalConstraints[Index].IsType<FPoleVectorConstraint>())
+						if (LimitConstraint.bZLimitSet)
 						{
-							// darw pole vector location
-							// draw 3 joints line and a plane to pole vector
-							FPoleVectorConstraint& Constraint = InternalConstraints[Index].Get<FPoleVectorConstraint>();
-							FTransform RootTransform = LinkData[Constraint.ParentBoneIndex].GetTransform();
-							FTransform JointTransform = LinkData[Constraint.BoneIndex].GetTransform();
-							FTransform ChildTransform = LinkData[Constraint.ChildBoneIndex].GetTransform();
-
-							FVector JointTarget = (Constraint.bUseLocalDir)? Constraint.CalculateCurrentPoleVectorDir(RootTransform, JointTransform, ChildTransform, LinkData[Constraint.BoneIndex].LocalFrame) : Constraint.PoleVector;
-
-							// draw the plane, 
-							TArray<FVector> Positions;
-							Positions.Add(RootTransform.GetLocation());
-							Positions.Add(ChildTransform.GetLocation());
-							Positions.Add(ChildTransform.GetLocation());
-							Positions.Add(JointTarget);
-							Positions.Add(JointTarget);
-							Positions.Add(RootTransform.GetLocation());
-
-							ExecuteContext.UnitContext.DrawInterface->DrawLines(Offset, Positions, FLinearColor::Gray, 1.2f);
-							ExecuteContext.UnitContext.DrawInterface->DrawLine(Offset, JointTransform.GetLocation(), JointTarget, FLinearColor::Red, 1.2f);
+							FTransform ZAxisConeTM(XAxis, ZAxis ^ XAxis, ZAxis, BaseTransform.GetTranslation());
+							ZAxisConeTM.SetRotation(FQuat(ZAxis, 0.f) * ZAxisConeTM.GetRotation());
+							ZAxisConeTM.SetScale3D(FVector(30.f));
+							ExecuteContext.UnitContext.DrawInterface->DrawCone(Offset, ZAxisConeTM, LimitConstraint.Limit.Z, 0.0f, 24, false, FLinearColor::Blue, GEngine->ConstraintLimitMaterialZ->GetRenderProxy());
 						}
 					}
+
+					if (InternalConstraints[Index].IsType<FPoleVectorConstraint>())
+					{
+						// darw pole vector location
+						// draw 3 joints line and a plane to pole vector
+						FPoleVectorConstraint& Constraint = InternalConstraints[Index].Get<FPoleVectorConstraint>();
+						FTransform RootTransform = LinkData[Constraint.ParentBoneIndex].GetTransform();
+						FTransform JointTransform = LinkData[Constraint.BoneIndex].GetTransform();
+						FTransform ChildTransform = LinkData[Constraint.ChildBoneIndex].GetTransform();
+
+						FVector JointTarget = (Constraint.bUseLocalDir)? Constraint.CalculateCurrentPoleVectorDir(RootTransform, JointTransform, ChildTransform, LinkData[Constraint.BoneIndex].LocalFrame) : Constraint.PoleVector;
+
+						// draw the plane, 
+						TArray<FVector> Positions;
+						Positions.Add(RootTransform.GetLocation());
+						Positions.Add(ChildTransform.GetLocation());
+						Positions.Add(ChildTransform.GetLocation());
+						Positions.Add(JointTarget);
+						Positions.Add(JointTarget);
+						Positions.Add(RootTransform.GetLocation());
+
+						ExecuteContext.UnitContext.DrawInterface->DrawLines(Offset, Positions, FLinearColor::Gray, 1.2f);
+						ExecuteContext.UnitContext.DrawInterface->DrawLine(Offset, JointTransform.GetLocation(), JointTarget, FLinearColor::Red, 1.2f);
+					}
 				}
 			}
-			///////////////////////////////////////////////////////////////////////////
-			// debug draw end
-			///////////////////////////////////////////////////////////////////////////
 		}
+		///////////////////////////////////////////////////////////////////////////
+		// debug draw end
+		///////////////////////////////////////////////////////////////////////////
+	}
 
-		// we update back to hierarchy
-		for (int32 LinkIndex= 0; LinkIndex < LinkData.Num() ; ++LinkIndex)
-		{
-			// only propagate, if you are leaf joints here
-			// this means, only the last joint in the test
-			const FRigElementKey& CurrentItem = *LinkDataToHierarchyIndices.Find(LinkIndex);
-			const FTransform& LinkTransform = LinkData[LinkIndex].GetTransform();
-			Hierarchy->SetGlobalTransform(CurrentItem, LinkTransform, false, bPropagateToChildren);
-		}
+	// we update back to hierarchy
+	for (int32 LinkIndex= 0; LinkIndex < LinkData.Num() ; ++LinkIndex)
+	{
+		// only propagate, if you are leaf joints here
+		// this means, only the last joint in the test
+		const FRigElementKey& CurrentItem = *LinkDataToHierarchyIndices.Find(LinkIndex);
+		const FTransform& LinkTransform = LinkData[LinkIndex].GetTransform();
+		Hierarchy->SetGlobalTransform(CurrentItem, LinkTransform, false, bPropagateToChildren);
 	}
 }
 
diff --git a/Engine/Plugins/Experimental/FullBodyIK/Source/PBIK/Private/RigUnit_PBIK.cpp b/Engine/Plugins/Experimental/FullBodyIK/Source/PBIK/Private/RigUnit_PBIK.cpp
--- a/Engine/Plugins/Experimental/FullBodyIK/Source/PBIK/Private/RigUnit_PBIK.cpp
+++ b/Engine/Plugins/Experimental/FullBodyIK/Source/PBIK/Private/RigUnit_PBIK.cpp
@@ -10,20 +10,12 @@ FRigUnit_PBIK_Execute()
 	DECLARE_SCOPE_HIERARCHICAL_COUNTER_RIGUNIT()
 	LLM_SCOPE_BYNAME(TEXT("Animation/FBIK"));
 
-	if (ExecuteContext.UnitContext.State == EControlRigState::Init)
+	if (bNeedsInit)
 	{
 		BoneSettingToSolverBoneIndex.Reset();
 		Solver.Reset();
 		SolverBoneToElementIndex.Reset();
 		EffectorSolverIndices.Reset();
-		bNeedsInit = true;
-		return;
-	}
-
-	// only updates from here on...
-	if (ExecuteContext.UnitContext.State != EControlRigState::Update)
-	{
-		return;
 	}
 
 	URigHierarchy* Hierarchy = ExecuteContext.Hierarchy;
diff --git a/Engine/Plugins/Experimental/LiveLinkControlRig/Source/LiveLinkControlRig/Private/LiveLinkRigUnits.cpp b/Engine/Plugins/Experimental/LiveLinkControlRig/Source/LiveLinkControlRig/Private/LiveLinkRigUnits.cpp
--- a/Engine/Plugins/Experimental/LiveLinkControlRig/Source/LiveLinkControlRig/Private/LiveLinkRigUnits.cpp
+++ b/Engine/Plugins/Experimental/LiveLinkControlRig/Source/LiveLinkControlRig/Private/LiveLinkRigUnits.cpp
@@ -25,159 +25,115 @@ ILiveLinkClient* LiveLinkControlRigUtilities::TryGetLiveLinkClient()
 
 FRigUnit_LiveLinkEvaluteFrameAnimation_Execute()
 {
-	switch (ExecuteContext.UnitContext.State)
-	{
-	case EControlRigState::Init:
-	case EControlRigState::Update:
-	{
-		// Ensure a default value in case the evaluation fails
-		SubjectFrame = FSubjectFrameHandle();
+	// Ensure a default value in case the evaluation fails
+	SubjectFrame = FSubjectFrameHandle();
 
-		// Try to update SubjectFrame from Live Link
-		if (ILiveLinkClient* LiveLinkClient = LiveLinkControlRigUtilities::TryGetLiveLinkClient())
+	// Try to update SubjectFrame from Live Link
+	if (ILiveLinkClient* LiveLinkClient = LiveLinkControlRigUtilities::TryGetLiveLinkClient())
+	{
+		FLiveLinkSubjectFrameData FrameData;
+		if (LiveLinkClient->EvaluateFrame_AnyThread(SubjectName, ULiveLinkAnimationRole::StaticClass(), FrameData))
 		{
-			FLiveLinkSubjectFrameData FrameData;
-			if (LiveLinkClient->EvaluateFrame_AnyThread(SubjectName, ULiveLinkAnimationRole::StaticClass(), FrameData))
-			{
-				FLiveLinkBlueprintDataStruct BlueprintStruct(FSubjectFrameHandle::StaticStruct(), &SubjectFrame);
-				GetDefault<ULiveLinkAnimationRole>()->InitializeBlueprintData(FrameData, BlueprintStruct);
-			}
+			FLiveLinkBlueprintDataStruct BlueprintStruct(FSubjectFrameHandle::StaticStruct(), &SubjectFrame);
+			GetDefault<ULiveLinkAnimationRole>()->InitializeBlueprintData(FrameData, BlueprintStruct);
 		}
+	}
 
-		// Debug Drawing if required
-		if (bDrawDebug && ExecuteContext.UnitContext.DrawInterface != nullptr)
-		{
-			/*
-			FRigBoneHierarchy DebugHierarchy;
+	// Debug Drawing if required
+	if (bDrawDebug && ExecuteContext.UnitContext.DrawInterface != nullptr)
+	{
+		/*
+		FRigBoneHierarchy DebugHierarchy;
 
-			for (int BoneIdx = 0; BoneIdx < SubjectFrame.GetNumberOfTransforms(); ++BoneIdx)
+		for (int BoneIdx = 0; BoneIdx < SubjectFrame.GetNumberOfTransforms(); ++BoneIdx)
+		{
+			FLiveLinkTransform CurrentTransform;
+			SubjectFrame.GetTransformByIndex(BoneIdx, CurrentTransform);
+			FName CurrentBoneName;
+			CurrentTransform.GetName(CurrentBoneName);
+			FName CurrentParentName(NAME_None);
+			if (CurrentTransform.HasParent())
 			{
-				FLiveLinkTransform CurrentTransform;
-				SubjectFrame.GetTransformByIndex(BoneIdx, CurrentTransform);
-				FName CurrentBoneName;
-				CurrentTransform.GetName(CurrentBoneName);
-				FName CurrentParentName(NAME_None);
-				if (CurrentTransform.HasParent())
-				{
-					FLiveLinkTransform ParentTransform;
-					CurrentTransform.GetParent(ParentTransform);
-					ParentTransform.GetName(CurrentParentName);
-				}
-				FTransform LocalSpaceTransform, GlobalSpaceTransform;
-				CurrentTransform.GetTransformParentSpace(LocalSpaceTransform);
-				CurrentTransform.GetTransformRootSpace(GlobalSpaceTransform);
-				DebugHierarchy.Add(CurrentBoneName, CurrentParentName, ERigBoneType::User, LocalSpaceTransform, LocalSpaceTransform, GlobalSpaceTransform);
+				FLiveLinkTransform ParentTransform;
+				CurrentTransform.GetParent(ParentTransform);
+				ParentTransform.GetName(CurrentParentName);
 			}
-			Context.DrawInterface->DrawHierarchy(DebugDrawOffset, DebugHierarchy, EControlRigDrawHierarchyMode::Axes, 1.0f, DebugColor, 1.0f);
-			*/
+			FTransform LocalSpaceTransform, GlobalSpaceTransform;
+			CurrentTransform.GetTransformParentSpace(LocalSpaceTransform);
+			CurrentTransform.GetTransformRootSpace(GlobalSpaceTransform);
+			DebugHierarchy.Add(CurrentBoneName, CurrentParentName, ERigBoneType::User, LocalSpaceTransform, LocalSpaceTransform, GlobalSpaceTransform);
 		}
-	}
-		break;
-	default:
-		checkNoEntry(); // Execute is only defined for Init and Update
-		break;
+		Context.DrawInterface->DrawHierarchy(DebugDrawOffset, DebugHierarchy, EControlRigDrawHierarchyMode::Axes, 1.0f, DebugColor, 1.0f);
+		*/
 	}
 }
 
 FRigUnit_LiveLinkGetTransformByName_Execute()
 {
-	switch (ExecuteContext.UnitContext.State)
-	{
-	case EControlRigState::Init:
-	case EControlRigState::Update:
-	{
-		// Ensure a default value in case the evaluation fails
-		Transform = FTransform::Identity;
+	// Ensure a default value in case the evaluation fails
+	Transform = FTransform::Identity;
 
-		// Try to update Transform from the Subject Frame
-		FLiveLinkTransform LiveLinkTransform;
-		SubjectFrame.GetTransformByName(TransformName, LiveLinkTransform);
-		switch (Space)
-		{
-		case EBoneGetterSetterMode::GlobalSpace:
-			LiveLinkTransform.GetTransformRootSpace(Transform);
-			break;
-		case EBoneGetterSetterMode::LocalSpace:
-			LiveLinkTransform.GetTransformParentSpace(Transform);
-			break;
-		default:
-			checkNoEntry(); // Unknown bone space. We should never hit this
-			break;
-		}
-	}
+	// Try to update Transform from the Subject Frame
+	FLiveLinkTransform LiveLinkTransform;
+	SubjectFrame.GetTransformByName(TransformName, LiveLinkTransform);
+	switch (Space)
+	{
+	case EBoneGetterSetterMode::GlobalSpace:
+		LiveLinkTransform.GetTransformRootSpace(Transform);
+		break;
+	case EBoneGetterSetterMode::LocalSpace:
+		LiveLinkTransform.GetTransformParentSpace(Transform);
 		break;
 	default:
-		checkNoEntry(); // Execute is only defined for Init and Update
+		checkNoEntry(); // Unknown bone space. We should never hit this
 		break;
 	}
 }
 
 FRigUnit_LiveLinkGetParameterValueByName_Execute()
 {
-	switch (ExecuteContext.UnitContext.State)
-	{
-	case EControlRigState::Init:
-	case EControlRigState::Update:
+	float NewValue;
+	if (SubjectFrame.GetCurveValueByName(ParameterName, NewValue))
 	{
-		float NewValue;
-		if (SubjectFrame.GetCurveValueByName(ParameterName, NewValue))
-		{
-			Value = NewValue;
-		}
-		else
-		{
-			Value = 0.0f;
-		}
+		Value = NewValue;
 	}
-		break;
-	default:
-		checkNoEntry(); // Execute is only defined for Init and Update
-		break;
+	else
+	{
+		Value = 0.0f;
 	}
 }
 
 FRigUnit_LiveLinkEvaluteFrameTransform_Execute()
 {
-	switch (ExecuteContext.UnitContext.State)
-	{
-	case EControlRigState::Init:
-	case EControlRigState::Update:
-	{
-		// Ensure a default value in case the evaluation fails
-		Transform = FTransform::Identity;
+	// Ensure a default value in case the evaluation fails
+	Transform = FTransform::Identity;
 
-		// Try to update Transform from Live Link
-		if (ILiveLinkClient* LiveLinkClient = LiveLinkControlRigUtilities::TryGetLiveLinkClient())
+	// Try to update Transform from Live Link
+	if (ILiveLinkClient* LiveLinkClient = LiveLinkControlRigUtilities::TryGetLiveLinkClient())
+	{
+		FLiveLinkSubjectFrameData FrameData;
+		if (LiveLinkClient->EvaluateFrame_AnyThread(SubjectName, ULiveLinkTransformRole::StaticClass(), FrameData))
 		{
-			FLiveLinkSubjectFrameData FrameData;
-			if (LiveLinkClient->EvaluateFrame_AnyThread(SubjectName, ULiveLinkTransformRole::StaticClass(), FrameData))
-			{
-				FLiveLinkTransformFrameData* TransformData = FrameData.FrameData.Cast<FLiveLinkTransformFrameData>();
-				Transform = TransformData->Transform;
-			}
+			FLiveLinkTransformFrameData* TransformData = FrameData.FrameData.Cast<FLiveLinkTransformFrameData>();
+			Transform = TransformData->Transform;
 		}
+	}
 
-		// Debug Drawing if required
-		if (bDrawDebug && ExecuteContext.UnitContext.DrawInterface != nullptr)
-		{
-			// Rather than drawing as an axis guide, draw as a hierarchy with a dummy root bone which makes it clearer what the origin of the transform
-			// space is
+	// Debug Drawing if required
+	if (bDrawDebug && ExecuteContext.UnitContext.DrawInterface != nullptr)
+	{
+		// Rather than drawing as an axis guide, draw as a hierarchy with a dummy root bone which makes it clearer what the origin of the transform
+		// space is
 
-			/*
-			FRigBoneHierarchy DebugHierarchy;
+		/*
+		FRigBoneHierarchy DebugHierarchy;
 
-			FName RootName("root");
-			FName TransformName("transform");
-			DebugHierarchy.Add(RootName, FName(NAME_None), ERigBoneType::User, FTransform::Identity, FTransform::Identity, FTransform::Identity);
-			DebugHierarchy.Add(TransformName, RootName, ERigBoneType::User, Transform, Transform, Transform);
+		FName RootName("root");
+		FName TransformName("transform");
+		DebugHierarchy.Add(RootName, FName(NAME_None), ERigBoneType::User, FTransform::Identity, FTransform::Identity, FTransform::Identity);
+		DebugHierarchy.Add(TransformName, RootName, ERigBoneType::User, Transform, Transform, Transform);
 
-			Context.DrawInterface->DrawHierarchy(DebugDrawOffset, DebugHierarchy, EControlRigDrawHierarchyMode::Axes, 1.0f, DebugColor, 1.0f);
-			*/
-		}
-	}
-		break;
-	default:
-		checkNoEntry(); // Execute is only defined for Init and Update
-		break;
+		Context.DrawInterface->DrawHierarchy(DebugDrawOffset, DebugHierarchy, EControlRigDrawHierarchyMode::Axes, 1.0f, DebugColor, 1.0f);
+		*/
 	}
 }
\ No newline at end of file
diff --git a/Engine/Source/Developer/RigVMDeveloper/Private/RigVMCompiler/RigVMCompiler.cpp b/Engine/Source/Developer/RigVMDeveloper/Private/RigVMCompiler/RigVMCompiler.cpp
--- a/Engine/Source/Developer/RigVMDeveloper/Private/RigVMCompiler/RigVMCompiler.cpp
+++ b/Engine/Source/Developer/RigVMDeveloper/Private/RigVMCompiler/RigVMCompiler.cpp
@@ -15,6 +15,7 @@
 #include "ProfilingDebugging/ScopedTimers.h"
 #include "RigVMModel/RigVMClient.h"
 #include "Algo/Count.h"
+#include "String/Join.h"
 
 #include UE_INLINE_GENERATED_CPP_BY_NAME(RigVMCompiler)
 
@@ -3458,6 +3459,20 @@ FRigVMOperand URigVMCompiler::FindOrAddRegister(const FRigVMVarExprAST* InVarExp
 				const FProperty* Property = UnitStruct->FindPropertyByName(Pin->GetFName());
 				check(Property);
 
+				JoinedDefaultValue.Reset();
+					
+				FStructOnScope StructOnScope(UnitStruct);
+				const FRigVMStruct* StructMemory = (const FRigVMStruct*)StructOnScope.GetStructMemory();
+				const uint8* PropertyMemory = Property->ContainerPtrToValuePtr<uint8>(StructMemory);
+				
+				Property->ExportText_Direct(
+					JoinedDefaultValue,
+					PropertyMemory,
+					PropertyMemory,
+					nullptr,
+					PPF_None,
+					nullptr);
+
 				if (!Property->HasMetaData(FRigVMStruct::SingletonMetaName))
 				{
 					bValidHiddenPin = true;
@@ -3469,6 +3484,7 @@ FRigVMOperand URigVMCompiler::FindOrAddRegister(const FRigVMVarExprAST* InVarExp
 				if(const FRigVMDispatchFactory* Factory = DispatchNode->GetFactory())
 				{
 					bValidHiddenPin = !Factory->HasArgumentMetaData(Pin->GetFName(), FRigVMStruct::SingletonMetaName);
+					JoinedDefaultValue = Factory->GetArgumentDefaultValue(Pin->GetFName(), Pin->GetTypeIndex());
 				}
 			}
 
diff --git a/Engine/Source/Runtime/RigVM/Private/RigVMCore/RigVM.cpp b/Engine/Source/Runtime/RigVM/Private/RigVMCore/RigVM.cpp
--- a/Engine/Source/Runtime/RigVM/Private/RigVMCore/RigVM.cpp
+++ b/Engine/Source/Runtime/RigVM/Private/RigVMCore/RigVM.cpp
@@ -1553,12 +1553,13 @@ bool URigVM::ShouldHaltAtInstruction(const FName& InEventName, const uint16 Inst
 }
 #endif
 
-bool URigVM::Initialize(TArrayView<URigVMMemoryStorage*> Memory, bool bInitializeMemory)
+bool URigVM::Initialize(TArrayView<URigVMMemoryStorage*> Memory)
 {
 	if (ExecutingThreadId != INDEX_NONE)
 	{
 		ensureMsgf(ExecutingThreadId == FPlatformTLS::GetCurrentThreadId(), TEXT("RigVM::Initialize from multiple threads (%d and %d)"), ExecutingThreadId, (int32)FPlatformTLS::GetCurrentThreadId());
 	}
+	
 	CopyDeferredVMIfRequired();
 	TGuardValue<int32> GuardThreadId(ExecutingThreadId, FPlatformTLS::GetCurrentThreadId());
 
@@ -1577,189 +1578,40 @@ bool URigVM::Initialize(TArrayView<URigVMMemoryStorage*> Memory, bool bInitializ
 		Memory = LocalMemory;
 	}
 
-	CacheMemoryHandlesIfRequired(Memory);
-
-	if(!bInitializeMemory)
-	{
-		return true;
-	}
-	
-	FRigVMByteCode& ByteCode = GetByteCode();
-	TArray<const FRigVMFunction*>& Functions = GetFunctions();
-	TArray<const FRigVMDispatchFactory*>& Factories = GetFactories();
-
-#if WITH_EDITOR
-	TArray<FName>& FunctionNames = GetFunctionNames();
-#endif
-
-	Context.Reset();
-	Context.SliceOffsets.AddZeroed(Instructions.Num());
-	FRigVMExecuteContext& ContextPublicData = Context.GetPublicData<>();
-
-	TGuardValue<URigVM*> VMInContext(Context.VM, this);
-	
-	while (Instructions.IsValidIndex(ContextPublicData.InstructionIndex))
+	// re-initialize work memory from CDO
+	if(URigVMMemoryStorage* WorkMemory = Memory[(int32)ERigVMMemoryType::Work])
 	{
-		const FRigVMInstruction& Instruction = Instructions[ContextPublicData.InstructionIndex];
-
-		switch (Instruction.OpCode)
+		if(!WorkMemory->HasAnyFlags(RF_ClassDefaultObject))
 		{
-			case ERigVMOpCode::Execute_0_Operands:
-			case ERigVMOpCode::Execute_1_Operands:
-			case ERigVMOpCode::Execute_2_Operands:
-			case ERigVMOpCode::Execute_3_Operands:
-			case ERigVMOpCode::Execute_4_Operands:
-			case ERigVMOpCode::Execute_5_Operands:
-			case ERigVMOpCode::Execute_6_Operands:
-			case ERigVMOpCode::Execute_7_Operands:
-			case ERigVMOpCode::Execute_8_Operands:
-			case ERigVMOpCode::Execute_9_Operands:
-			case ERigVMOpCode::Execute_10_Operands:
-			case ERigVMOpCode::Execute_11_Operands:
-			case ERigVMOpCode::Execute_12_Operands:
-			case ERigVMOpCode::Execute_13_Operands:
-			case ERigVMOpCode::Execute_14_Operands:
-			case ERigVMOpCode::Execute_15_Operands:
-			case ERigVMOpCode::Execute_16_Operands:
-			case ERigVMOpCode::Execute_17_Operands:
-			case ERigVMOpCode::Execute_18_Operands:
-			case ERigVMOpCode::Execute_19_Operands:
-			case ERigVMOpCode::Execute_20_Operands:
-			case ERigVMOpCode::Execute_21_Operands:
-			case ERigVMOpCode::Execute_22_Operands:
-			case ERigVMOpCode::Execute_23_Operands:
-			case ERigVMOpCode::Execute_24_Operands:
-			case ERigVMOpCode::Execute_25_Operands:
-			case ERigVMOpCode::Execute_26_Operands:
-			case ERigVMOpCode::Execute_27_Operands:
-			case ERigVMOpCode::Execute_28_Operands:
-			case ERigVMOpCode::Execute_29_Operands:
-			case ERigVMOpCode::Execute_30_Operands:
-			case ERigVMOpCode::Execute_31_Operands:
-			case ERigVMOpCode::Execute_32_Operands:
-			case ERigVMOpCode::Execute_33_Operands:
-			case ERigVMOpCode::Execute_34_Operands:
-			case ERigVMOpCode::Execute_35_Operands:
-			case ERigVMOpCode::Execute_36_Operands:
-			case ERigVMOpCode::Execute_37_Operands:
-			case ERigVMOpCode::Execute_38_Operands:
-			case ERigVMOpCode::Execute_39_Operands:
-			case ERigVMOpCode::Execute_40_Operands:
-			case ERigVMOpCode::Execute_41_Operands:
-			case ERigVMOpCode::Execute_42_Operands:
-			case ERigVMOpCode::Execute_43_Operands:
-			case ERigVMOpCode::Execute_44_Operands:
-			case ERigVMOpCode::Execute_45_Operands:
-			case ERigVMOpCode::Execute_46_Operands:
-			case ERigVMOpCode::Execute_47_Operands:
-			case ERigVMOpCode::Execute_48_Operands:
-			case ERigVMOpCode::Execute_49_Operands:
-			case ERigVMOpCode::Execute_50_Operands:
-			case ERigVMOpCode::Execute_51_Operands:
-			case ERigVMOpCode::Execute_52_Operands:
-			case ERigVMOpCode::Execute_53_Operands:
-			case ERigVMOpCode::Execute_54_Operands:
-			case ERigVMOpCode::Execute_55_Operands:
-			case ERigVMOpCode::Execute_56_Operands:
-			case ERigVMOpCode::Execute_57_Operands:
-			case ERigVMOpCode::Execute_58_Operands:
-			case ERigVMOpCode::Execute_59_Operands:
-			case ERigVMOpCode::Execute_60_Operands:
-			case ERigVMOpCode::Execute_61_Operands:
-			case ERigVMOpCode::Execute_62_Operands:
-			case ERigVMOpCode::Execute_63_Operands:
-			case ERigVMOpCode::Execute_64_Operands:
+			if(const URigVMMemoryStorageGeneratorClass* MemoryClass = Cast<URigVMMemoryStorageGeneratorClass>(WorkMemory->GetClass()))
 			{
-				const FRigVMExecuteOp& Op = ByteCode.GetOpAt<FRigVMExecuteOp>(Instruction);
-				int32 OperandCount = FirstHandleForInstruction[ContextPublicData.InstructionIndex + 1] - FirstHandleForInstruction[ContextPublicData.InstructionIndex];
-				FRigVMMemoryHandleArray OpHandles;
-				if(OperandCount > 0)
-				{
-					OpHandles = FRigVMMemoryHandleArray(&CachedMemoryHandles[FirstHandleForInstruction[ContextPublicData.InstructionIndex]], OperandCount);
-				}
-#if WITH_EDITOR
-				ContextPublicData.FunctionName = FunctionNames[Op.FunctionIndex];
-#endif
-				Context.Factory = Factories[Op.FunctionIndex];
-
-				// find out the largest slice count
-				int32 MaxSliceCount = 1;
-
-				const bool bContainsLazyValue = OpHandles.ContainsByPredicate([](const FRigVMMemoryHandle& Handle)
-				{
-					return Handle.IsLazy();
-				});
-
-				if(!bContainsLazyValue)
+				if(URigVMMemoryStorage* WorkMemoryCDO = MemoryClass->GetDefaultObject<URigVMMemoryStorage>())
 				{
-					Context.BeginSlice(MaxSliceCount);
-					for (int32 SliceIndex = 0; SliceIndex < MaxSliceCount; SliceIndex++)
+					for(const FProperty* Property : MemoryClass->LinkedProperties)
 					{
-						(*Functions[Op.FunctionIndex]->FunctionPtr)(Context, OpHandles);
-						Context.IncrementSlice();
+						Property->CopyCompleteValue_InContainer(WorkMemory, WorkMemoryCDO);
+#if UE_RIGVM_DEBUG_EXECUTION
+						FString DefaultValue;
+						const uint8* PropertyMemory = Property->ContainerPtrToValuePtr<uint8>(WorkMemory);
+						Property->ExportText_Direct(
+							DefaultValue,
+							PropertyMemory,
+							PropertyMemory,
+							nullptr,
+							PPF_None,
+							nullptr);
+
+						UE_LOG(LogRigVM, Warning, TEXT("Property %s defaults to '%s'."), *Property->GetName(), *DefaultValue);
+						UE_LOG(LogRigVM, Warning, TEXT("Property %s defaults to '%s'."), *Property->GetName(), *DefaultValue);
+#endif
 					}
-					Context.EndSlice();
 				}
-				break;
-			}
-			case ERigVMOpCode::Zero:
-			case ERigVMOpCode::BoolFalse:
-			case ERigVMOpCode::BoolTrue:
-			{
-				break;
-			}
-			case ERigVMOpCode::Copy:
-			{
-				const FRigVMCopyOp& Op = ByteCode.GetOpAt<FRigVMCopyOp>(Instruction);
-
-				FRigVMMemoryHandle& SourceHandle = CachedMemoryHandles[FirstHandleForInstruction[ContextPublicData.InstructionIndex]];
-				FRigVMMemoryHandle& TargetHandle = CachedMemoryHandles[FirstHandleForInstruction[ContextPublicData.InstructionIndex] + 1];
-				URigVMMemoryStorage::CopyProperty(TargetHandle, SourceHandle);
-				break;
-			}
-			case ERigVMOpCode::Increment:
-			case ERigVMOpCode::Decrement:
-			case ERigVMOpCode::Equals:
-			case ERigVMOpCode::NotEquals:
-			case ERigVMOpCode::JumpAbsolute:
-			case ERigVMOpCode::JumpForward:
-			case ERigVMOpCode::JumpBackward:
-			case ERigVMOpCode::JumpAbsoluteIf:
-			case ERigVMOpCode::JumpForwardIf:
-			case ERigVMOpCode::JumpBackwardIf:
-			case ERigVMOpCode::ChangeType:
-			case ERigVMOpCode::BeginBlock:
-			case ERigVMOpCode::EndBlock:
-			case ERigVMOpCode::Exit:
-			case ERigVMOpCode::ArrayGetNum:
-			case ERigVMOpCode::ArraySetNum:
-			case ERigVMOpCode::ArrayAppend:
-			case ERigVMOpCode::ArrayClone:
-			case ERigVMOpCode::ArrayGetAtIndex:
-			case ERigVMOpCode::ArraySetAtIndex:
-			case ERigVMOpCode::ArrayInsert:
-			case ERigVMOpCode::ArrayRemove:
-			case ERigVMOpCode::ArrayAdd:
-			case ERigVMOpCode::ArrayFind:
-			case ERigVMOpCode::ArrayIterator:
-			case ERigVMOpCode::ArrayUnion:
-			case ERigVMOpCode::ArrayDifference:
-			case ERigVMOpCode::ArrayIntersection:
-			case ERigVMOpCode::ArrayReverse:
-			case ERigVMOpCode::InvokeEntry:
-			case ERigVMOpCode::JumpToBranch:
-			{
-				break;
-			}
-			case ERigVMOpCode::Invalid:
-			{
-				ensure(false);
-				return false;
 			}
 		}
-		ContextPublicData.InstructionIndex++;
 	}
 	
+	CacheMemoryHandlesIfRequired(Memory);
+
 	return true;
 }
 
diff --git a/Engine/Source/Runtime/RigVM/Private/RigVMCore/RigVMNativized.cpp b/Engine/Source/Runtime/RigVM/Private/RigVMCore/RigVMNativized.cpp
--- a/Engine/Source/Runtime/RigVM/Private/RigVMCore/RigVMNativized.cpp
+++ b/Engine/Source/Runtime/RigVM/Private/RigVMCore/RigVMNativized.cpp
@@ -27,7 +27,7 @@ void URigVMNativized::Reset(bool IsIgnoringArchetypeRef)
 	ByteCodeStorage.Reset();
 }
 
-bool URigVMNativized::Initialize(TArrayView<URigVMMemoryStorage*> Memory, bool bInitializeMemory)
+bool URigVMNativized::Initialize(TArrayView<URigVMMemoryStorage*> Memory)
 {
 	// nothing to do here 
 	return true;
diff --git a/Engine/Source/Runtime/RigVM/Public/RigVMCore/RigVM.h b/Engine/Source/Runtime/RigVM/Public/RigVMCore/RigVM.h
--- a/Engine/Source/Runtime/RigVM/Public/RigVMCore/RigVM.h
+++ b/Engine/Source/Runtime/RigVM/Public/RigVMCore/RigVM.h
@@ -210,7 +210,7 @@ public:
 	}
 
 	// Initializes all execute ops and their memory.
-	virtual bool Initialize(TArrayView<URigVMMemoryStorage*> Memory, bool bInitializeMemory = true);
+	virtual bool Initialize(TArrayView<URigVMMemoryStorage*> Memory);
 
 	// Executes the VM.
 	// You can optionally provide external memory to the execution
diff --git a/Engine/Source/Runtime/RigVM/Public/RigVMCore/RigVMExecuteContext.h b/Engine/Source/Runtime/RigVM/Public/RigVMCore/RigVMExecuteContext.h
--- a/Engine/Source/Runtime/RigVM/Public/RigVMCore/RigVMExecuteContext.h
+++ b/Engine/Source/Runtime/RigVM/Public/RigVMCore/RigVMExecuteContext.h
@@ -239,7 +239,9 @@ struct RIGVM_API FRigVMExecuteContext
 		{
 			(*RuntimeSettings.LogFunction)(InSeverity, this, InMessage);
 		}
+#if !UE_RIGVM_DEBUG_EXECUTION
 		else
+#endif
 		{
 			if(InSeverity == EMessageSeverity::Error)
 			{
@@ -273,6 +275,9 @@ struct RIGVM_API FRigVMExecuteContext
 
 	double GetAbsoluteTime() const { return AbsoluteTime; } 
 	void SetAbsoluteTime(double InAbsoluteTime) { AbsoluteTime = InAbsoluteTime; }
+
+	double GetFramesPerSecond() const { return FramesPerSecond; } 
+	void SetFramesPerSecond(double InFramesPerSecond) { FramesPerSecond = InFramesPerSecond; }
 	
 	FRigVMNameCache* GetNameCache() const { return NameCache; }
 
@@ -311,6 +316,8 @@ protected:
 
 	double AbsoluteTime;
 
+	double FramesPerSecond;
+
 	FRigVMRuntimeSettings RuntimeSettings;
 
 	mutable FRigVMNameCache* NameCache;
diff --git a/Engine/Source/Runtime/RigVM/Public/RigVMCore/RigVMMemoryStorage.h b/Engine/Source/Runtime/RigVM/Public/RigVMCore/RigVMMemoryStorage.h
--- a/Engine/Source/Runtime/RigVM/Public/RigVMCore/RigVMMemoryStorage.h
+++ b/Engine/Source/Runtime/RigVM/Public/RigVMCore/RigVMMemoryStorage.h
@@ -9,6 +9,10 @@
 #include "EdGraph/EdGraphNode.h"
 #include "Misc/AssertionMacros.h"
 #include "RigVMMemoryCommon.h"
+#include "RigVMDefines.h"
+#if UE_RIGVM_DEBUG_EXECUTION
+#include "RigVMModule.h"
+#endif
 #include "RigVMPropertyPath.h"
 #include "RigVMStatistics.h"
 #include "RigVMTraits.h"
@@ -809,9 +813,26 @@ private:
 					FScriptArrayHelper DefaultArrayHelper(ArrayProperty, DefaultArrayMemory);
 					if (const uint8* DefaultElementMemory = DefaultArrayHelper.GetRawPtr(0))
 					{
-						FProperty* ElementProperty = ArrayProperty->Inner;
+						const FProperty* ElementProperty = ArrayProperty->Inner;
 						for (int32 i=FirstAddedIndex; i<ArrayHelper.Num(); ++i)
 						{
+#if UE_RIGVM_DEBUG_EXECUTION
+							FString DefaultValue;
+							ElementProperty->ExportText_Direct(
+								DefaultValue,
+								DefaultElementMemory,
+								DefaultElementMemory,
+								nullptr,
+								PPF_None,
+								nullptr);
+
+							UE_LOG(LogRigVM, Display, TEXT("Adding slice %d for Property '%s', defaulting to '%s'."),
+								InSliceIndex,
+								*ArrayProperty->GetName(),
+								*DefaultValue
+							);
+#endif
+
 							uint8* DestMemory = ArrayHelper.GetRawPtr(i);
 							ElementProperty->CopyCompleteValue(DestMemory, DefaultElementMemory);
 						}
@@ -819,6 +840,25 @@ private:
 				}
 			}
 
+#if UE_RIGVM_DEBUG_EXECUTION
+			const FProperty* ElementProperty = ArrayProperty->Inner;
+			FString DefaultValue;
+			const uint8* DefaultElementMemory = ArrayHelper.GetRawPtr(InSliceIndex);
+			ElementProperty->ExportText_Direct(
+				DefaultValue,
+				DefaultElementMemory,
+				DefaultElementMemory,
+				nullptr,
+				PPF_None,
+				nullptr);
+
+			UE_LOG(LogRigVM, Display, TEXT("Getting slice %d for Property '%s', currently '%s'."),
+				InSliceIndex,
+				*ArrayProperty->GetName(),
+				*DefaultValue
+			);
+#endif
+
 			return ArrayHelper.GetRawPtr(InSliceIndex);
 		}
 
diff --git a/Engine/Source/Runtime/RigVM/Public/RigVMCore/RigVMNativized.h b/Engine/Source/Runtime/RigVM/Public/RigVMCore/RigVMNativized.h
--- a/Engine/Source/Runtime/RigVM/Public/RigVMCore/RigVMNativized.h
+++ b/Engine/Source/Runtime/RigVM/Public/RigVMCore/RigVMNativized.h
@@ -50,7 +50,7 @@ public:
 	virtual bool IsNativized() const override { return true; }
 	virtual void Empty() override { return; }
 	virtual void CopyFrom(URigVM* InVM, bool bDeferCopy = false, bool bReferenceLiteralMemory = false, bool bReferenceByteCode = false, bool bCopyExternalVariables = false, bool bCopyDynamicRegisters = false) override { return; }
-	virtual bool Initialize(TArrayView<URigVMMemoryStorage*> Memory, bool bInitializeMemory = true) override;
+	virtual bool Initialize(TArrayView<URigVMMemoryStorage*> Memory) override;
 	virtual ERigVMExecuteResult Execute(TArrayView<URigVMMemoryStorage*> Memory, const FName& InEntryName = NAME_None) override;
 	virtual int32 AddRigVMFunction(UScriptStruct* InRigVMStruct, const FName& InMethodName) override { return INDEX_NONE; }
 	virtual FString GetRigVMFunctionName(int32 InFunctionIndex) const override { return FString(); }
