Various improvements to FBIK solver:

Added over-relaxation to core constraint solve.
Added hard clamp on joint constraint rotation delta per-iteration.
Added better "PrePull" root behavior using deformation tensor of goal transforms to generate "best fit" rotation.
Removed "Min Mass Multiplier" as rotation clamping now prevents divergence.
Fixed bug potentially causing "lock root" mode to not work correctly.

diff --git a/Engine/Plugins/Animation/IKRig/Source/IKRig/Private/Solvers/IKRig_PBIKSolver.cpp b/Engine/Plugins/Animation/IKRig/Source/IKRig/Private/Solvers/IKRig_PBIKSolver.cpp
--- a/Engine/Plugins/Animation/IKRig/Source/IKRig/Private/Solvers/IKRig_PBIKSolver.cpp
+++ b/Engine/Plugins/Animation/IKRig/Source/IKRig/Private/Solvers/IKRig_PBIKSolver.cpp
@@ -117,10 +117,13 @@ void UIKRigPBIKSolver::Solve(FIKRigSkeleton& IKRigSkeleton, const FIKRigGoalCont
 	// update settings
 	FPBIKSolverSettings Settings;
 	Settings.Iterations = Iterations;
-	Settings.bAllowStretch = bAllowStretch;
 	Settings.MassMultiplier = MassMultiplier;
-	Settings.MinMassMultiplier = MinMassMultiplier;
+	Settings.bAllowStretch = bAllowStretch;
 	Settings.RootBehavior = RootBehavior;
+	Settings.PrePullRootSettings = PrePullRootSettings;
+	Settings.GlobalPullChainAlpha = PullChainAlpha;
+	Settings.MaxAngle = MaxAngle;
+	Settings.OverRelaxation = OverRelaxation;
 	Settings.bStartSolveFromInputPose = bStartSolveFromInputPose;
 
 	// solve
@@ -149,10 +152,13 @@ void UIKRigPBIKSolver::UpdateSolverSettings(UIKRigSolver* InSettings)
 	if(UIKRigPBIKSolver* Settings = Cast<UIKRigPBIKSolver>(InSettings))
 	{
 		Iterations = Settings->Iterations;
-		bAllowStretch = Settings->bAllowStretch;
 		MassMultiplier = Settings->MassMultiplier;
-		MinMassMultiplier = Settings->MinMassMultiplier;
+		bAllowStretch = Settings->bAllowStretch;
 		RootBehavior = Settings->RootBehavior;
+		PrePullRootSettings = Settings->PrePullRootSettings;
+		PullChainAlpha = Settings->PullChainAlpha;
+		MaxAngle = Settings->MaxAngle;
+		OverRelaxation = Settings->OverRelaxation;
 		bStartSolveFromInputPose = Settings->bStartSolveFromInputPose;
 
 		// copy effector settings
diff --git a/Engine/Plugins/Animation/IKRig/Source/IKRig/Public/Solvers/IKRig_PBIKSolver.h b/Engine/Plugins/Animation/IKRig/Source/IKRig/Public/Solvers/IKRig_PBIKSolver.h
--- a/Engine/Plugins/Animation/IKRig/Source/IKRig/Public/Solvers/IKRig_PBIKSolver.h
+++ b/Engine/Plugins/Animation/IKRig/Source/IKRig/Public/Solvers/IKRig_PBIKSolver.h
@@ -188,30 +188,44 @@ public:
 	FName RootBone;
 
 	/** High iteration counts can help solve complex joint configurations with competing constraints, but will increase runtime cost. Default is 20. */
-	UPROPERTY(EditAnywhere, Category = "Full Body IK Settings", meta = (ClampMin = "0", ClampMax = "1000", UIMin = "0.0", UIMax = "200.0"))
+	UPROPERTY(EditAnywhere, Category = SolverSettings, meta = (ClampMin = "0", ClampMax = "1000", UIMin = "0.0", UIMax = "200.0"))
 	int32 Iterations = 20;
 
 	/** A global mass multiplier; higher values will make the joints more stiff, but require more iterations. Typical range is 0.0 to 10.0. */
-	UPROPERTY(EditAnywhere, Category = "Full Body IK Settings", meta = (ClampMin = "0", UIMin = "0.0", UIMax = "10.0"))
+	UPROPERTY(EditAnywhere, Category = SolverSettings, meta = (ClampMin = "0", UIMin = "0.0", UIMax = "10.0"))
 	float MassMultiplier = 1.0f;
 
-	/** Set this as low as possible while keeping the solve stable. Lower values improve convergence of effector targets. Default is 0.2. */
-	UPROPERTY(EditAnywhere, Category = "Full Body IK Settings", meta = (ClampMin = "0", UIMin = "0.0", UIMax = "10.0"))
-	float MinMassMultiplier = 0.2f;
-
 	/** If true, joints will translate to reach the effectors; causing bones to lengthen if necessary. Good for cartoon effects. Default is false. */
-	UPROPERTY(EditAnywhere, Category = "Full Body IK Settings")
+	UPROPERTY(EditAnywhere, Category = SolverSettings)
 	bool bAllowStretch = false;
 
-	/** (Default is PrePull) Set the behavior for the translation of the root.
-	*Pre Pull: translates the whole body by the average motion of the stretched effectors to help achieve faster convergence when reaching far.
-	*Pin to Input: locks the translation and rotation of the root bone to the input pose. Overrides any bone settings applied to the root.
+	/** (Default is PrePull) Set the behavior of the solver root.
+	*Pre Pull: translates and rotates the root (and all children) by the average motion of the stretched effectors to help achieve faster convergence when reaching far.
+	*Pin to Input: locks the translation and rotation of the root bone to the input pose. Overrides any bone settings applied to the root. Good for partial-body solves.
 	*Free: treats the root bone like any other and allows it to move freely or according to any bone settings applied to it. */
-	UPROPERTY(EditAnywhere, Category = "Full Body IK Settings")
-	EPBIKRootBehavior RootBehavior;
+	UPROPERTY(EditAnywhere, Category = RootBehavior)
+	EPBIKRootBehavior RootBehavior = EPBIKRootBehavior::PrePull;
+
+	/** Settings only applicable when Root Behavior is set to "PrePull". Use these values to adjust the gross movement and orientation of the entire skeleton. */
+	UPROPERTY(EditAnywhere, Category = RootBehavior)
+	FRootPrePullSettings PrePullRootSettings;
+
+	/** A global multiplier for all Pull Chain Alpha values on all effectors. Range is 0.0 to 1.0. Default is 1.0. */
+	UPROPERTY(EditAnywhere, Category = AdvancedSettings, meta = (ClampMin = "0", ClampMax = "1", UIMin = "0.0", UIMax = "1.0"))
+	float PullChainAlpha = 1.0f;
+
+	/** Maximum angle that a joint can be rotated per constraint iteration. Lower this value if the solve is diverging. Range is 0.0 to 180.0. Default is 30. */
+	UPROPERTY(EditAnywhere, Category = AdvancedSettings, meta = (ClampMin = "0", ClampMax = "45", UIMin = "0.0", UIMax = "180.0"))
+	float MaxAngle = 30.f;
+
+	/** Pushes constraints beyond their normal amount to speed up convergence. Increasing this may speed up convergence, but at the cost of stability. Range is 1.0 - 2.0. Default is 1.3. */
+	UPROPERTY(EditAnywhere, Category = AdvancedSettings, meta = (ClampMin = "1", UIMin = "1.0", UIMax = "10.0"))
+	float OverRelaxation = 1.3f;
 	
-	/** When true, the solver is reset each tick to start from the current input pose. If false, incoming animated poses are ignored and the solver starts from the results of the previous solve. Default is true. */
-	UPROPERTY(EditAnywhere, Category = "Full Body IK Settings")
+	/** When true, the solver is reset each tick to start from the current input pose. Default is true.
+	 * If false, incoming animated poses are ignored and the solver starts from the results of the previous solve.
+	 * In very limited circumstances, it can be beneficial to use the pose from the previous frame.*/
+	UPROPERTY(EditAnywhere, Category = AdvancedSettings)
 	bool bStartSolveFromInputPose = true;
 	
 	UPROPERTY()
diff --git a/Engine/Plugins/Experimental/FullBodyIK/Source/PBIK/Private/Core/PBIKBody.cpp b/Engine/Plugins/Experimental/FullBodyIK/Source/PBIK/Private/Core/PBIKBody.cpp
--- a/Engine/Plugins/Experimental/FullBodyIK/Source/PBIK/Private/Core/PBIKBody.cpp
+++ b/Engine/Plugins/Experimental/FullBodyIK/Source/PBIK/Private/Core/PBIKBody.cpp
@@ -93,10 +93,11 @@ void FRigidBody::UpdateFromInputs(const FPBIKSolverSettings& Settings)
 		InputPosition = Position;
 	}
 
-	// for fork joints (multiple solved children) we sum lengths to all children (see Initialize)
-	const float MinMass = 0.5f; // prevent mass ever hitting zero
-	MaxInvMass = 1.0f / (Mass * ((Settings.MassMultiplier * GLOBAL_UNITS) + MinMass));
-	MinInvMass = 1.0f / (Mass * ((Settings.MinMassMultiplier * GLOBAL_UNITS) + MinMass));
+	// update InvMass based on global mass multiplier
+	constexpr float MinMass = 0.5f; // prevent mass ever hitting zero
+	InvMass = 1.0f / FMath::Max(MinMass,(Mass * Settings.MassMultiplier * GLOBAL_UNITS));
+
+	SolverSettings = &Settings;
 }
 
 int FRigidBody::GetNumBonesToRoot() const
@@ -114,6 +115,16 @@ FRigidBody* FRigidBody::GetParentBody() const
 	return nullptr;
 }
 
+float FRigidBody::GetInverseMass()
+{
+	if (Pin && Pin->bEnabled)
+	{
+		return 1.0f - Pin->Alpha;
+	}
+
+	return InvMass;
+}
+
 void FRigidBody::ApplyPushToRotateBody(const FVector& Push, const FVector& Offset)
 {
 	if (Pin && Pin->bEnabled && Pin->bPinRotation)
@@ -129,7 +140,7 @@ void FRigidBody::ApplyPushToRotateBody(const FVector& Push, const FVector& Offse
 
 void FRigidBody::ApplyPushToPosition(const FVector& Push)
 {
-	Position += Push * (1.0f - J.PositionStiffness);
+	Position += Push * (1.0f - J.PositionStiffness) * SolverSettings->OverRelaxation;
 }
 
 void FRigidBody::ApplyRotationDelta(const FQuat& DeltaQ)
@@ -139,8 +150,17 @@ void FRigidBody::ApplyRotationDelta(const FQuat& DeltaQ)
 		return; // rotation of this body is pinned
 	}
 
+	// limit rotation each iteration
+	FQuat ClampedDQ = DeltaQ;
+	const float MaxPhi = FMath::DegreesToRadians(SolverSettings->MaxAngle);
+	const float Phi = DeltaQ.Size();
+	if (Phi > MaxPhi)
+	{
+		ClampedDQ *= MaxPhi / Phi;
+	}
+
 	/** DeltaQ is assumed to be a "pure" quaternion representing an infintesimal rotation */
-	FQuat Delta = DeltaQ * Rotation;
+	FQuat Delta = ClampedDQ * Rotation;
 	Delta.X *= 0.5f;
 	Delta.Y *= 0.5f;
 	Delta.Z *= 0.5f;
diff --git a/Engine/Plugins/Experimental/FullBodyIK/Source/PBIK/Private/Core/PBIKConstraint.cpp b/Engine/Plugins/Experimental/FullBodyIK/Source/PBIK/Private/Core/PBIKConstraint.cpp
--- a/Engine/Plugins/Experimental/FullBodyIK/Source/PBIK/Private/Core/PBIKConstraint.cpp
+++ b/Engine/Plugins/Experimental/FullBodyIK/Source/PBIK/Private/Core/PBIKConstraint.cpp
@@ -2,6 +2,7 @@
 
 #include "Core/PBIKConstraint.h"
 #include "Core/PBIKBody.h"
+#include "Core/PBIKSolver.h"
 
 namespace PBIK
 {
@@ -22,7 +23,7 @@ FJointConstraint::FJointConstraint(FRigidBody* InA, FRigidBody* InB)
 	AngleX = AngleY = AngleZ = 0;
 }
 	
-void FJointConstraint::Solve(bool bMoveSubRoots)
+void FJointConstraint::Solve(const FPBIKSolverSettings& Settings)
 {
 	// get pos correction to use for rotation
 	FVector OffsetA;
@@ -38,42 +39,37 @@ void FJointConstraint::Solve(bool bMoveSubRoots)
 	UpdateJointLimits();
 
 	// calc inv mass of body A
-	float AInvMass = 1.0f;
-	AInvMass -= (A->Pin && A->Pin->bEnabled) ? A->Pin->Alpha : 0.0f;
-	AInvMass -= !bMoveSubRoots && A->Bone->bIsSubRoot ? 1.0f : 0.0f;
-	AInvMass = AInvMass <= 0.0f ? 0.0f : AInvMass;
-
-	// calc inv mass of body B
-	float BInvMass = 1.0f;
-	BInvMass -= (B->Pin && B->Pin->bEnabled) ? B->Pin->Alpha : 0.0f;
-	BInvMass -= !bMoveSubRoots && B->Bone->bIsSubRoot ? 1.0f : 0.0f;
-	BInvMass = BInvMass <= 0.0f ? 0.0f : BInvMass;
-
-	const float InvMassSum = AInvMass + BInvMass;
-	if (FMath::IsNearlyZero(InvMassSum))
+	const float WA = A->GetInverseMass();
+	const float WB = B->GetInverseMass();
+	const float W = WA + WB;
+	if (FMath::IsNearlyZero(W))
 	{
 		return; // no correction can be applied on full locked configuration
 	}
-	const float OneOverInvMassSum = 1.0f / InvMassSum;
 
 	// apply positional correction to align pin point on both Bodies
 	// NOTE: applying position correction AFTER rotation takes into
 	// consideration change in relative pin locations mid-step after
 	// bodies have been rotated by ApplyPushToRotateBody
 	Correction = GetPositionCorrection(OffsetA, OffsetB);
-	Correction = Correction * OneOverInvMassSum;
+	FVector N;
+	float C;
+	Correction.ToDirectionAndLength(N, C);
 
-	A->ApplyPushToPosition(Correction * AInvMass);
-	B->ApplyPushToPosition(-Correction * BInvMass);
+	const FVector CorrectA = N * (C * (WA / W));
+	const FVector CorrectB = N * (C * (WB / W));
+
+	A->ApplyPushToPosition(CorrectA);
+	B->ApplyPushToPosition(-CorrectB);
 }
 
-void FJointConstraint::RemoveStretch()
+void FJointConstraint::RemoveStretch(const float Percent)
 {
 	// apply position correction, keeping parent body fixed
 	FVector ToA;
 	FVector ToB;
 	const FVector Correction = GetPositionCorrection(ToA, ToB);
-	B->Position -= Correction;
+	B->Position -= Correction * Percent;
 }
 
 FVector FJointConstraint::GetPositionCorrection(FVector& OutBodyToA, FVector& OutBodyToB) const
@@ -281,30 +277,29 @@ FPinConstraint::FPinConstraint(
 	bPinRotation = bInPinRotation;
 }
 
-void FPinConstraint::Solve(bool bMoveSubRoots)
+void FPinConstraint::Solve(const FPBIKSolverSettings& Settings)
 {
 	if (!bEnabled || (Alpha <= KINDA_SMALL_NUMBER))
 	{
 		return;
 	}
 
-	if (!bMoveSubRoots && A->Bone->bIsSubRoot)
-	{
-		return;
-	}
-
-	// get positional correction for rotation
-	FVector AToPinPoint;
-	FVector Correction = GetPositionCorrection(AToPinPoint);
-
 	if (bPinRotation)
 	{
-		// move body to pin location
-		A->Position += Correction;
 		// keep body at fixed rotation relative to the pin goal rotation
 		A->Rotation = ARotLocalToPin * GoalRotation;
+		
+		// move body to pin location
+		FVector AToPinPoint;
+		const FVector Correction = GetPositionCorrection(AToPinPoint);
+		A->Position += Correction;
+		
 	}else
 	{
+		// get positional correction for rotation
+		FVector AToPinPoint;
+		FVector Correction = GetPositionCorrection(AToPinPoint);
+		
 		// rotate body from alignment of pin points
 		A->ApplyPushToRotateBody(Correction, AToPinPoint);
 
@@ -312,7 +307,7 @@ void FPinConstraint::Solve(bool bMoveSubRoots)
 		// (applying directly without considering PositionStiffness because PinConstraints need
 		// to precisely pull the attached body to achieve convergence)
 		Correction = GetPositionCorrection(AToPinPoint);
-		A->Position += Correction; 
+		A->Position += Correction * Settings.OverRelaxation; 
 	}
 }
 
@@ -323,6 +318,13 @@ void FPinConstraint::SetGoal(const FVector& InGoalPosition, const FQuat& InGoalR
 	Alpha = InAlpha;
 }
 
+void FPinConstraint::EnableInCurrentState()
+{
+	GoalPosition = A->Position;
+	GoalRotation = A->Rotation;
+	bEnabled = true;
+}
+
 FVector FPinConstraint::GetPositionCorrection(FVector& OutBodyToPinPoint) const
 {
 	OutBodyToPinPoint = A->Rotation * PinPointLocalToA;
diff --git a/Engine/Plugins/Experimental/FullBodyIK/Source/PBIK/Private/Core/PBIKSolver.cpp b/Engine/Plugins/Experimental/FullBodyIK/Source/PBIK/Private/Core/PBIKSolver.cpp
--- a/Engine/Plugins/Experimental/FullBodyIK/Source/PBIK/Private/Core/PBIKSolver.cpp
+++ b/Engine/Plugins/Experimental/FullBodyIK/Source/PBIK/Private/Core/PBIKSolver.cpp
@@ -189,41 +189,144 @@
 	UpdateBonesFromBodies();
 }
 
-void FPBIKSolver::PullRootTowardsEffectors()
+void FPBIKSolver::UpdateBodies(const FPBIKSolverSettings& Settings)
+{
+	using PBIK::FRigidBody;
+	
+	// apply large scale gross movement to entire skeleton based on effector deltas
+	// (this creates a better pose from which to start the constraint solving)
+	ApplyRootPrePull(Settings.RootBehavior, Settings.PrePullRootSettings);
+
+	// pre-rotate bones towards preferred angles when effector limb is squashed
+	ApplyPreferredAngles();
+
+	// pull each effector's chain towards itself
+	ApplyPullChainAlpha(Settings.GlobalPullChainAlpha);
+		
+	// run ALL constraint iterations
+	SolveConstraints(Settings);
+}
+
+void FPBIKSolver::UpdateBonesFromBodies()
+{
+	using PBIK::FRigidBody;
+	using PBIK::FBone;
+	using PBIK::FEffector;
+	
+	// update Bone transforms controlled by Bodies
+	for (FRigidBody& Body : Bodies)
+	{
+		Body.Bone->Position = Body.Position + Body.Rotation * Body.BoneLocalPosition;
+		Body.Bone->Rotation = Body.Rotation;
+	}
+
+	// update Bone transforms controlled by effectors
+	for (const FEffector& Effector : Effectors)
+	{
+		FBone* Bone = Effector.Bone;
+		if (Bone->bIsSolverRoot)
+		{
+			continue; // if there's an effector on the root, leave it where the body ended up
+		}
+
+		if (Bone->Body)
+		{
+			continue; // effector is between other effectors, so leave transform where body ended up
+		}
+		
+		Bone->Position = Bone->Parent->Position + Bone->Parent->Rotation * Bone->LocalPositionOrig;
+
+		if (Effector.Settings.PinRotation > SMALL_NUMBER)
+		{
+			// optionally pin rotation to that of effector
+			const float RotAmount = FMath::Clamp(0.0f, 1.0f, Effector.Settings.PinRotation);
+			Bone->Rotation = FQuat::FastLerp(Bone->Rotation, Effector.Rotation, RotAmount).GetNormalized();
+		}else
+		{
+			Bone->Rotation = Bone->Parent->Rotation * Bone->LocalRotationOrig;
+		}
+	}
+
+	// propagate to non-solved bones (requires storage in root to tip order)
+	for (FBone& Bone : Bones)
+	{
+		if (Bone.bIsSolved || !Bone.Parent)
+		{
+			continue;
+		}
+
+		Bone.Position = Bone.Parent->Position + Bone.Parent->Rotation * Bone.LocalPositionOrig;
+		Bone.Rotation = Bone.Parent->Rotation * Bone.LocalRotationOrig;
+	}
+}
+
+void FPBIKSolver::ApplyRootPrePull(const EPBIKRootBehavior RootBehavior, const FRootPrePullSettings& PrePullSettings)
 {
 	using PBIK::FEffector;
 	using PBIK::FRigidBody;
 	
-	// accumulate a delta offset from each effector
-	FVector BodyOffset = FVector::ZeroVector;
+	if (RootBehavior != EPBIKRootBehavior::PrePull)
+	{
+		return;
+	}
+
+	// calculate a "best fit" transform from deformed goal locations
+	TArray<FVector, TInlineAllocator<8>> InitialPoints;
+	TArray<FVector, TInlineAllocator<8>> CurrentPoints;
+	InitialPoints.Reserve(Effectors.Num());
+	CurrentPoints.Reserve(Effectors.Num());
 	for (FEffector& Effector : Effectors)
 	{
-		const float DistToRootStraightLine = (Effector.Position - SolverRoot->Position).Size();
-		const float Delta = DistToRootStraightLine - Effector.DistToRootStraightLine;//Effector.DistToRootAlongBones;
-		if (Delta < SMALL_NUMBER)
-		{
-			continue; // only pull
-		}
-
-		FVector RootToTip = Effector.Position - SolverRoot->Position;
-		BodyOffset += Delta * RootToTip.GetSafeNormal();
+		InitialPoints.Add(Effector.PositionOrig);
+		CurrentPoints.Add(Effector.Position);
 	}
 
-	const float InvNumEffectors = (1.0f / static_cast<float>(Effectors.Num()));
-	BodyOffset *= InvNumEffectors;
+	// calculate "best fit" rotational delta from deformed effector configuration
+	FVector InitialCentroid;
+	FVector CurrentCentroid;
+	const FQuat RotationOffset = GetRotationFromDeformedPoints(
+		InitialPoints,
+		CurrentPoints,
+		InitialCentroid,
+		CurrentCentroid);
 
-	// linearly move entire apparatus towards tip effectors
+	// do per-axis alpha blend for position offset
+	FVector PositionOffset = CurrentCentroid - InitialCentroid;
+	PositionOffset *= FVector(PrePullSettings.PositionAlphaX, PrePullSettings.PositionAlphaY, PrePullSettings.PositionAlphaZ);
+	// alpha blend entire offset
+	PositionOffset *= PrePullSettings.PositionAlpha;
+	
+	// do per-axis alpha blend for rotation offset
+	const FVector Euler = RotationOffset.Euler() * FVector(PrePullSettings.RotationAlphaX, PrePullSettings.RotationAlphaY, PrePullSettings.RotationAlphaZ);
+	FQuat FinalRotationOffset = FQuat::MakeFromEuler(Euler);
+	// alpha blend the entire rotation offset
+	FinalRotationOffset = FQuat::FastLerp(FQuat::Identity, FinalRotationOffset, PrePullSettings.RotationAlpha).GetNormalized();
+
+	// move and rotate all bodies rigidly with the root
+	const FRigidBody* RootBody = SolverRoot->Body;
+	const FVector RootPosition = RootBody->Position;
+	const FQuat RootRotation = RootBody->Rotation;
+	const FVector NewRootPosition = RootPosition + PositionOffset;
+	const FQuat NewRootRotation = FinalRotationOffset * RootBody->Rotation;
 	for (FRigidBody& Body : Bodies)
 	{
-		Body.Position += BodyOffset;
+		FVector LocalPos = RootRotation.Inverse() * (Body.Position - RootPosition);
+		FQuat LocalRot = RootRotation.Inverse() * Body.Rotation;
+		Body.Position = NewRootPosition + (NewRootRotation * LocalPos);
+		Body.Rotation = NewRootRotation * LocalRot;
 	}
 }
 
-void FPBIKSolver::ApplyPullChainAlpha()
+void FPBIKSolver::ApplyPullChainAlpha(const float GlobalPullChainAlpha)
 {
 	using PBIK::FEffector;
 	using PBIK::FBone;
 	using PBIK::FRigidBody;
+
+	if (FMath::IsNearlyZero(GlobalPullChainAlpha))
+	{
+		return;
+	}
 	
 	for (FEffector& Effector : Effectors)
 	{
@@ -267,7 +370,7 @@
 		(ChainEndNew - ChainStartNew).ToDirectionAndLength(ChainVecNew, ChainLenNew);
 
 		FQuat ChainDeltaRotation = FQuat::FindBetweenNormals(ChainVecOrig, ChainVecNew);
-		const float ChainDeltaAlpha = Effector.Settings.PullChainAlpha * Effector.Settings.StrengthAlpha;
+		const float ChainDeltaAlpha = Effector.Settings.PullChainAlpha * Effector.Settings.StrengthAlpha * GlobalPullChainAlpha;
 		ChainDeltaRotation = FQuat::FastLerp(FQuat::Identity, ChainDeltaRotation, ChainDeltaAlpha);
 		FVector ChainDeltaPosition = ChainVecNew * (ChainLenNew - ChainLenOrig) * ChainDeltaAlpha;
 
@@ -304,114 +407,35 @@
 	// preferred angles introduce stretch, so remove that to prevent biasing the first constraint iteration
 	for (int32 C = Constraints.Num() - 1; C >= 0; --C)
 	{
-		Constraints[C]->RemoveStretch();
+		constexpr float FullAmount = 1.0f;
+		Constraints[C]->RemoveStretch(FullAmount);
 	}
 }
 
-void FPBIKSolver::UpdateBodies(const FPBIKSolverSettings& Settings)
-{
-	// this creates a better pose from which to start the constraint solving
-	if (Settings.RootBehavior == EPBIKRootBehavior::PrePull)
-	{
-		// pull ALL bodies towards the effectors
-		PullRootTowardsEffectors();
-	}
-
-	// pre-rotate bones towards preferred angles when effector limb is squashed
-	ApplyPreferredAngles();
-
-	// pull each effector's chain towards itself
-	ApplyPullChainAlpha();
-		
-	// run ALL constraint iterations
-	SolveConstraints(Settings.Iterations, true, Settings.bAllowStretch);
-}
-
-void FPBIKSolver::SolveConstraints(const int32 Iterations, const bool bMoveRoots, const bool bAllowStretch)
+void FPBIKSolver::SolveConstraints(const FPBIKSolverSettings& Settings)
 {
 	using PBIK::FRigidBody;
-	
-	for (int32  I = 0; I < Iterations; ++I)
+
+	for (int32  I = 0; I < Settings.Iterations; ++I)
 	{
-		float IterPercent = Iterations == 1 ? 0 : (static_cast<float>(I) / static_cast<float>(Iterations - 1));
-		
-		// gradually ramp down mass of bodies
-		float MassRamp = PBIK::SquaredEaseOut(IterPercent);
-		for (FRigidBody& Body : Bodies)
-		{
-			Body.InvMass = FMath::Lerp(Body.MaxInvMass, Body.MinInvMass, MassRamp);
-		}
-		
 		// solve all constraints
-		for (auto Constraint : Constraints)
+		for (const TSharedPtr<PBIK::FConstraint>& Constraint : Constraints)
 		{
-			Constraint->Solve(bMoveRoots);
+			Constraint->Solve(Settings);
 		}
-
+		
 		// do post-pass to remove stretch
-		if (!bAllowStretch && (I>Iterations/2))
+		if (!Settings.bAllowStretch)
 		{
+			const float StretchAmount = FMath::Min(1.0f, static_cast<float>(I+1) / (Settings.Iterations / 2.0f));
 			for (int32 C = Constraints.Num() - 1; C >= 0; --C)
 			{
-				Constraints[C]->RemoveStretch();
+				Constraints[C]->RemoveStretch(StretchAmount);
 			}
 		}
 	}
 }
 
-void FPBIKSolver::UpdateBonesFromBodies()
-{
-	using PBIK::FRigidBody;
-	using PBIK::FBone;
-	using PBIK::FEffector;
-	
-	// update Bone transforms controlled by Bodies
-	for (FRigidBody& Body : Bodies)
-	{
-		Body.Bone->Position = Body.Position + Body.Rotation * Body.BoneLocalPosition;
-		Body.Bone->Rotation = Body.Rotation;
-	}
-
-	// update Bone transforms controlled by effectors
-	for (const FEffector& Effector : Effectors)
-	{
-		FBone* Bone = Effector.Bone;
-		if (Bone->bIsSolverRoot)
-		{
-			continue; // if there's an effector on the root, leave it where the body ended up
-		}
-
-		if (Bone->Body)
-		{
-			continue; // effector is between other effectors, so leave transform where body ended up
-		}
-		
-		Bone->Position = Bone->Parent->Position + Bone->Parent->Rotation * Bone->LocalPositionOrig;
-
-		if (Effector.Settings.PinRotation > SMALL_NUMBER)
-		{
-			// optionally pin rotation to that of effector
-			const float RotAmount = FMath::Clamp(0.0f, 1.0f, Effector.Settings.PinRotation);
-			Bone->Rotation = FQuat::FastLerp(Bone->Rotation, Effector.Rotation, RotAmount).GetNormalized();
-		}else
-		{
-			Bone->Rotation = Bone->Parent->Rotation * Bone->LocalRotationOrig;
-		}
-	}
-
-	// propagate to non-solved bones (requires storage in root to tip order)
-	for (FBone& Bone : Bones)
-	{
-		if (Bone.bIsSolved || !Bone.Parent)
-		{
-			continue;
-		}
-
-		Bone.Position = Bone.Parent->Position + Bone.Parent->Rotation * Bone.LocalPositionOrig;
-		Bone.Rotation = Bone.Parent->Rotation * Bone.LocalRotationOrig;
-	}
-}
-
 bool FPBIKSolver::Initialize()
 {
 	LLM_SCOPE_BYNAME(TEXT("Animation/FBIK"));
@@ -592,7 +616,7 @@
 	// must be done AFTER setting: Bone.IsSubRoot/IsSolverRoot/Parent
 	for (FEffector& Effector : Effectors)
 	{
-		FBone* Parent = Effector.Bone->Parent;
+		const FBone* Parent = Effector.Bone->Parent;
 		while (Parent)
 		{
 			if (!Parent->bIsSolved)
@@ -634,11 +658,12 @@
 	using PBIK::FJointConstraint;
 
 	Constraints.Empty();
-
-	// pin bodies to effectors
+	
+	// Add Effector Pin Constraints
+	// (these are added first since they create the pulling effect which needs resolved by the joint constraints)
 	for (FEffector& Effector : Effectors)
 	{
-		FBone* BodyBone = Effector.Bone->bIsSolverRoot ? Effector.Bone : Effector.Bone->Parent;
+		const FBone* BodyBone = Effector.Bone->bIsSolverRoot ? Effector.Bone : Effector.Bone->Parent;
 		if (!BodyBone)
 		{
 			UE_LOG(LogPBIKSolver, Warning, TEXT("PBIK: effector is on bone that does not have a parent."));
@@ -653,17 +678,8 @@
 		Constraints.Add(Constraint);
 	}
 
-	// pin root body to animated location 
-	// this constraint is by default off in solver settings
-	if (!SolverRoot->Body->Effector) // only add if user hasn't added their own root effector
-	{
-		const TSharedPtr<FPinConstraint> RootConstraint = MakeShared<FPinConstraint>(SolverRoot->Body, SolverRoot->Position, SolverRoot->Rotation, true);
-		Constraints.Add(RootConstraint);
-		RootPin = RootConstraint;
-		SolverRoot->Body->Pin = RootConstraint.Get();
-	}
-
-	// constrain all bodies together (child to parent)
+	// Add joint constraints
+	// (constrain all bodies together (child to parent))
 	for (FRigidBody& Body : Bodies)
 	{
 		FRigidBody* ParentBody = Body.GetParentBody();
@@ -676,6 +692,17 @@
 		Constraints.Add(Constraint);
 	}
 
+	// pin root body to animated location 
+	// this constraint is by default off in solver settings
+	// this constraint is added last to ensure it is fully solved at the end of the iteration loop
+	if (!SolverRoot->Body->Effector) // only add if user hasn't added their own root effector
+	{
+		const TSharedPtr<FPinConstraint> RootConstraint = MakeShared<FPinConstraint>(SolverRoot->Body, SolverRoot->Position, SolverRoot->Rotation, true);
+		Constraints.Add(RootConstraint);
+		RootPin = RootConstraint;
+		SolverRoot->Body->Pin = RootConstraint.Get();
+	}
+	
 	return true;
 }
 
@@ -785,3 +812,80 @@
 	Effectors[Index].SetGoal(InPosition, InRotation, InSettings);
 }
 
+// given a set of points in an initial configuration and the same set of points in a deformed configuration,
+// this function outputs a quaternion that represents the "best fit" rotation that rotates the initial points to the
+// current points.
+FQuat FPBIKSolver::GetRotationFromDeformedPoints(
+	const TArrayView<FVector>& InInitialPoints,
+	const TArrayView<FVector>& InCurrentPoints,
+	FVector& OutInitialCentroid,
+	FVector& OutCurrentCentroid)
+{
+	check(InInitialPoints.Num() == InCurrentPoints.Num());
+
+	// must have more than 1 point to generate a gradient
+	if (InInitialPoints.Num() <= 1)
+	{
+		return FQuat::Identity;
+	}
+	
+	// calculate initial and current centroids
+	OutInitialCentroid = FVector::ZeroVector;
+	OutCurrentCentroid = FVector::ZeroVector;
+	for (int32 PointIndex=0; PointIndex<InInitialPoints.Num(); ++PointIndex)
+	{
+		OutInitialCentroid += InInitialPoints[PointIndex];
+		OutCurrentCentroid += InCurrentPoints[PointIndex];
+	}
+
+	// average centroids
+	const float InvNumEffectors = 1.0f / static_cast<float>(InInitialPoints.Num());
+	OutInitialCentroid *= InvNumEffectors;
+	OutCurrentCentroid *= InvNumEffectors;
+
+	// accumulate deformation gradient to extract a rotation from
+	FVector DX = FVector::ZeroVector; // DX, DY, DZ are rows of 3x3 deformation gradient tensor
+	FVector DY = FVector::ZeroVector;
+	FVector DZ = FVector::ZeroVector;
+	for (int32 PointIndex=0; PointIndex<InInitialPoints.Num(); ++PointIndex)
+	{
+		// accumulate the deformation gradient tensor for all points
+		// "Meshless Deformations Based on Shape Matching"
+		// Equation 7 describes accumulation of deformation gradient from points
+		//
+		// P is normalized vector from INITIAL centroid to INITIAL point
+		// Q is normalized vector from CURRENT centroid to CURRENT point
+		FVector P = (InInitialPoints[PointIndex] - OutInitialCentroid).GetSafeNormal();
+		FVector Q = (InCurrentPoints[PointIndex] - OutCurrentCentroid).GetSafeNormal();
+		// PQ^T is the outer product of P and Q which is a 3x3 matrix
+		// https://en.m.wikipedia.org/wiki/Outer_product
+		DX += FVector(P[0]*Q[0], P[0]*Q[1], P[0]*Q[2]);
+		DY += FVector(P[1]*Q[0], P[1]*Q[1], P[1]*Q[2]);
+		DZ += FVector(P[2]*Q[0], P[2]*Q[1], P[2]*Q[2]);
+	}
+
+	// extract "best fit" rotation from deformation gradient
+	FQuat Q = FQuat::Identity;
+	constexpr int32 MaxIter = 50;
+	// "A Robust Method to Extract the Rotational Part of Deformations" equation 7
+	// https://matthias-research.github.io/pages/publications/stablePolarDecomp.pdf
+	for (unsigned int Iter = 0; Iter < MaxIter; Iter++)
+	{
+		FMatrix R = FRotationMatrix::Make(Q);
+		FVector RCol0(R.M[0][0], R.M[0][1], R.M[0][2]);
+		FVector RCol1(R.M[1][0], R.M[1][1], R.M[1][2]);
+		FVector RCol2(R.M[2][0], R.M[2][1], R.M[2][2]);
+		FVector Omega = RCol0.Cross(DX) + RCol1.Cross(DY) + RCol2.Cross(DZ);
+		Omega *= 1.0f / (fabs(RCol0.Dot(DX) + RCol1.Dot(DY) + RCol2.Dot(DZ)) + SMALL_NUMBER);
+		const float W = Omega.Size();
+		if (W < SMALL_NUMBER)
+		{
+			break;
+		}
+		Q = FQuat(FQuat((1.0 / W) * Omega, W)) * Q;
+		Q.Normalize();
+	}
+
+	return Q;
+}
+
diff --git a/Engine/Plugins/Experimental/FullBodyIK/Source/PBIK/Public/Core/PBIKBody.h b/Engine/Plugins/Experimental/FullBodyIK/Source/PBIK/Public/Core/PBIKBody.h
--- a/Engine/Plugins/Experimental/FullBodyIK/Source/PBIK/Public/Core/PBIKBody.h
+++ b/Engine/Plugins/Experimental/FullBodyIK/Source/PBIK/Public/Core/PBIKBody.h
@@ -80,6 +80,7 @@ struct FRigidBody
 	FBoneSettings J;
 	FPinConstraint* Pin = nullptr;
 	FEffector* Effector = nullptr;
+	const FPBIKSolverSettings* SolverSettings;
 
 	FVector Position;
 	FQuat Rotation;
@@ -89,9 +90,9 @@ struct FRigidBody
 	TArray<FVector> ChildLocalPositions;
 
 	float InvMass = 0.f;
-	float MaxInvMass = 0.f;
-	float MinInvMass = 0.f;
 	float Mass = 0.f;
+
+	float MaxAngle = 0.0f;
 	
 private:
 
@@ -109,6 +110,8 @@ public:
 
 	FRigidBody* GetParentBody() const;
 
+	float GetInverseMass();
+
 	void ApplyPushToRotateBody(const FVector& Push, const FVector& Offset);
 	
 	void ApplyPushToPosition(const FVector& Push);
diff --git a/Engine/Plugins/Experimental/FullBodyIK/Source/PBIK/Public/Core/PBIKConstraint.h b/Engine/Plugins/Experimental/FullBodyIK/Source/PBIK/Public/Core/PBIKConstraint.h
--- a/Engine/Plugins/Experimental/FullBodyIK/Source/PBIK/Public/Core/PBIKConstraint.h
+++ b/Engine/Plugins/Experimental/FullBodyIK/Source/PBIK/Public/Core/PBIKConstraint.h
@@ -11,8 +11,8 @@ struct FConstraint
 {
 	bool bEnabled = true;
 
-	virtual void Solve(bool bMoveSubRoots) = 0;
-	virtual void RemoveStretch() {};
+	virtual void Solve(const FPBIKSolverSettings& Settings) = 0;
+	virtual void RemoveStretch(const float Percent){};
 };
 
 struct FJointConstraint : public FConstraint
@@ -51,9 +51,9 @@ public:
 
 	virtual ~FJointConstraint() {};
 
-	virtual void Solve(bool bMoveSubRoots) override;
+	virtual void Solve(const FPBIKSolverSettings& Settings) override;
 
-	virtual void RemoveStretch() override;
+	virtual void RemoveStretch(const float Percent) override;
 
 private:
 
@@ -106,10 +106,12 @@ public:
 
 	virtual ~FPinConstraint() {};
 
-	virtual void Solve(bool bMoveSubRoots) override;
+	virtual void Solve(const FPBIKSolverSettings& Settings) override;
 
 	void SetGoal(const FVector& InGoalPosition, const FQuat& InGoalRotation, const float InAlpha);
 
+	void EnableInCurrentState();
+
 private:
 
 	FVector GetPositionCorrection(FVector& OutBodyToPinPoint) const;
diff --git a/Engine/Plugins/Experimental/FullBodyIK/Source/PBIK/Public/Core/PBIKSolver.h b/Engine/Plugins/Experimental/FullBodyIK/Source/PBIK/Public/Core/PBIKSolver.h
--- a/Engine/Plugins/Experimental/FullBodyIK/Source/PBIK/Public/Core/PBIKSolver.h
+++ b/Engine/Plugins/Experimental/FullBodyIK/Source/PBIK/Public/Core/PBIKSolver.h
@@ -14,6 +14,7 @@ namespace PBIK
 {
 
 static float GLOBAL_UNITS = 100.0f; // (1.0f = meters), (100.0f = centimeters)
+static float MIN_MASS = 0.5f;
 
 // A long tail ease out function. Input range, 0-1. 
 FORCEINLINE static float QuarticEaseOut(const float& Input){ return (FMath::Pow(Input-1.0f, 4.0f) * -1.0f) + 1.0f; };
@@ -99,6 +100,44 @@ enum class EPBIKRootBehavior : uint8
 	Free,
 };
 
+USTRUCT(BlueprintType)
+struct FRootPrePullSettings
+{
+	GENERATED_BODY()
+	
+	/** Range 0-1, default is 0. Apply a large scale rotation offset to the entire body prior to constraint solving.*/
+	UPROPERTY(EditAnywhere, Category = RootRotation, meta = (ClampMin = "0.0", ClampMax = "1.0", UIMin = "0.0", UIMax = "1.0"))
+	float RotationAlpha = 0.0f;
+
+	/** Range 0-1, default is 1. Blend contribution to rotation offset in the X axis in component space.*/
+	UPROPERTY(EditAnywhere, Category = RootRotation, meta = (ClampMin = "0.0", ClampMax = "1.0", UIMin = "0.0", UIMax = "1.0"))
+	float RotationAlphaX = 1.0f;
+
+	/** Range 0-1, default is 1. Blend contribution to rotation offset in the Y axis in component space.*/
+	UPROPERTY(EditAnywhere, Category = RootRotation, meta = (ClampMin = "0.0", ClampMax = "1.0", UIMin = "0.0", UIMax = "1.0"))
+	float RotationAlphaY = 1.0f;
+
+	/** Range 0-1, default is 1. Blend contribution to rotation offset in the Z axis in component space.*/
+	UPROPERTY(EditAnywhere, Category = RootRotation, meta = (ClampMin = "0.0", ClampMax = "1.0", UIMin = "0.0", UIMax = "1.0"))
+	float RotationAlphaZ = 1.0f;
+
+	/** Range 0-1, default is 1. Apply a large scale position offset to the entire body prior to constraint solving.*/
+	UPROPERTY(EditAnywhere, Category = RootPosition, meta = (ClampMin = "0.0", ClampMax = "1.0", UIMin = "0.0", UIMax = "1.0"))
+	float PositionAlpha = 1.0f;
+
+	/** Range 0-1, default is 1. Blend contribution to position offset in the X axis in component space.*/
+	UPROPERTY(EditAnywhere, Category = RootPosition, meta = (ClampMin = "0.0", ClampMax = "1.0", UIMin = "0.0", UIMax = "1.0"))
+	float PositionAlphaX = 1.0f;
+
+	/** Range 0-1, default is 1. Blend contribution to position offset in the Y axis in component space.*/
+	UPROPERTY(EditAnywhere, Category = RootPosition, meta = (ClampMin = "0.0", ClampMax = "1.0", UIMin = "0.0", UIMax = "1.0"))
+	float PositionAlphaY = 1.0f;
+
+	/** Range 0-1, default is 1. Blend contribution to position offset in the Z axis in component space.*/
+	UPROPERTY(EditAnywhere, Category = RootPosition, meta = (ClampMin = "0.0", ClampMax = "1.0", UIMin = "0.0", UIMax = "1.0"))
+	float PositionAlphaZ = 1.0f;
+};
+
 USTRUCT(BlueprintType)
 struct PBIK_API FPBIKSolverSettings
 {
@@ -112,23 +151,37 @@ struct PBIK_API FPBIKSolverSettings
 	UPROPERTY(EditAnywhere, Category = SolverSettings, meta = (ClampMin = "0", UIMin = "0.0", UIMax = "10.0"))
 	float MassMultiplier = 1.0f;
 
-	/** Set this as low as possible while keeping the solve stable. Lower values improve convergence of effector targets. Default is 0.2. */
-	UPROPERTY(EditAnywhere, Category = SolverSettings, meta = (ClampMin = "0", UIMin = "0.0", UIMax = "10.0"))
-	float MinMassMultiplier = 0.2f;
-
 	/** If true, joints will translate to reach the effectors; causing bones to lengthen if necessary. Good for cartoon effects. Default is false. */
 	UPROPERTY(EditAnywhere, Category = SolverSettings)
 	bool bAllowStretch = false;
 
 	/** (Default is PrePull) Set the behavior of the solver root.
-	*Pre Pull: translates the root (and all children) by the average motion of the stretched effectors to help achieve faster convergence when reaching far.
+	*Pre Pull: translates and rotates the root (and all children) by the average motion of the stretched effectors to help achieve faster convergence when reaching far.
 	*Pin to Input: locks the translation and rotation of the root bone to the input pose. Overrides any bone settings applied to the root. Good for partial-body solves.
 	*Free: treats the root bone like any other and allows it to move freely or according to any bone settings applied to it. */
-	UPROPERTY(EditAnywhere, Category = SolverSettings)
+	UPROPERTY(EditAnywhere, Category = RootBehavior)
 	EPBIKRootBehavior RootBehavior = EPBIKRootBehavior::PrePull;
 
-	/** When true, the solver is reset each tick to start from the current input pose. If false, incoming animated poses are ignored and the solver starts from the results of the previous solve. Default is true. */
-	UPROPERTY(EditAnywhere, Category = SolverSettings)
+	/** Settings only applicable when Root Behavior is set to "PrePull". Use these values to adjust the gross movement and orientation of the entire skeleton. */
+	UPROPERTY(EditAnywhere, Category = RootBehavior)
+	FRootPrePullSettings PrePullRootSettings;
+
+	/** A global multiplier for all Pull Chain Alpha values on all effectors. Range is 0.0 to 1.0. Default is 1.0. */
+	UPROPERTY(EditAnywhere, Category = AdvancedSettings, meta = (ClampMin = "0", ClampMax = "1", UIMin = "0.0", UIMax = "1.0"))
+	float GlobalPullChainAlpha = 1.0f;
+
+	/** Maximum angle that a joint can be rotated per constraint iteration. Lower this value if the solve is diverging. Range is 0.0 to 180.0. Default is 30. */
+	UPROPERTY(EditAnywhere, Category = AdvancedSettings, meta = (ClampMin = "0", ClampMax = "45", UIMin = "0.0", UIMax = "180.0"))
+	float MaxAngle = 30.f;
+
+	/** Pushes constraints beyond their normal amount to speed up convergence. Increasing this may speed up convergence, but at the cost of stability. Range is 1.0 - 2.0. Default is 1.3. */
+	UPROPERTY(EditAnywhere, Category = AdvancedSettings, meta = (ClampMin = "1", UIMin = "1.0", UIMax = "2.0"))
+	float OverRelaxation = 1.3f;
+	
+	/** When true, the solver is reset each tick to start from the current input pose. Default is true.
+	 * If false, incoming animated poses are ignored and the solver starts from the results of the previous solve.
+	 * In very limited circumstances, it can be beneficial to use the pose from the previous frame.*/
+	UPROPERTY(EditAnywhere, Category = AdvancedSettings)
 	bool bStartSolveFromInputPose = true;
 };
 
@@ -200,16 +253,25 @@ private:
 
 	void UpdateBonesFromBodies();
 
-	void SolveConstraints(const int32 Iterations, const bool bMoveRoots, const bool bAllowStretch);
+	void SolveConstraints(const FPBIKSolverSettings& Settings);
 	
-	void PullRootTowardsEffectors();
+	void ApplyRootPrePull(const EPBIKRootBehavior RootBehavior, const FRootPrePullSettings& PrePullSettings);
 	
-	void ApplyPullChainAlpha();
+	void ApplyPullChainAlpha(const float GlobalPullChainAlpha);
 
 	void ApplyPreferredAngles();
 
 private:
 
+	// given a set of points in an initial configuration and the same set of points in a deformed configuration,
+	// this function outputs a quaternion that represents the "best fit" rotation that rotates the initial points to the
+	// current points.
+	static FQuat GetRotationFromDeformedPoints(
+		const TArrayView<FVector>& InInitialPoints,
+		const TArrayView<FVector>& InCurrentPoints,
+		FVector& OutInitialCentroid,
+		FVector& OutCurrentCentroid);
+	
 	PBIK::FBone* SolverRoot = nullptr;
 	TWeakPtr<PBIK::FPinConstraint> RootPin = nullptr;
 	TArray<PBIK::FBone> Bones;
