Adding dynamic translation support to FBIK solver.

#swarm https://p4-swarm.epicgames.net/reviews/24087148

diff --git a/Engine/Plugins/Experimental/FullBodyIK/Source/PBIK/Private/Core/PBIKBody.cpp b/Engine/Plugins/Experimental/FullBodyIK/Source/PBIK/Private/Core/PBIKBody.cpp
--- a/Engine/Plugins/Experimental/FullBodyIK/Source/PBIK/Private/Core/PBIKBody.cpp
+++ b/Engine/Plugins/Experimental/FullBodyIK/Source/PBIK/Private/Core/PBIKBody.cpp
@@ -52,30 +52,17 @@ FRigidBody::FRigidBody(FBone* InBone)
 	J = FBoneSettings();
 }
 
-void FRigidBody::Initialize(FBone* SolverRoot)
+void FRigidBody::Initialize(const FBone* SolverRoot)
 {
-	FVector Centroid = Bone->Position;
-	Mass = 0.0f;
-	for(const FBone* Child : Bone->Children)
-	{
-		Centroid += Child->Position;
-		Mass += (Bone->Position - Child->Position).Size();
-	}
-	Centroid = Centroid * (1.0f / (Bone->Children.Num() + 1.0f));
-
-	Position = Centroid;
-	Rotation = InitialRotation = Bone->Rotation;
-	BoneLocalPosition = Rotation.Inverse() * (Bone->Position - Centroid);
-
-	for (FBone* Child : Bone->Children)
-	{
-		FVector ChildLocalPos = Rotation.Inverse() * (Child->Position - Centroid);
-		ChildLocalPositions.Add(ChildLocalPos);
-	}
+	// calculate transform and mass of body based on the skeleton
+	UpdateTransformAndMassFromBones();
 
+	// store initial rotation from ref pose
+	InitialRotation = Rotation;
+	
 	// calculate num bones distance to root
 	NumBonesToRoot = 0;
-	FBone* Parent = Bone;
+	const FBone* Parent = Bone;
 	while (Parent && Parent != SolverRoot)
 	{
 		NumBonesToRoot += 1;
@@ -87,10 +74,7 @@ void FRigidBody::UpdateFromInputs(const FPBIKSolverSettings& Settings)
 {
 	if (Settings.bStartSolveFromInputPose)
 	{
-		// set to input pose
-		Position = Bone->Position - Bone->Rotation * BoneLocalPosition;
-		Rotation = Bone->Rotation;
-		InputPosition = Position;
+		UpdateTransformAndMassFromBones();
 	}
 
 	// update InvMass based on global mass multiplier
@@ -100,6 +84,29 @@ void FRigidBody::UpdateFromInputs(const FPBIKSolverSettings& Settings)
 	SolverSettings = &Settings;
 }
 
+void FRigidBody::UpdateTransformAndMassFromBones()
+{
+	FVector Centroid = Bone->Position;
+	Mass = 0.0f;
+	for(const FBone* Child : Bone->Children)
+	{
+		Centroid += Child->Position;
+		Mass += (Bone->Position - Child->Position).Size();
+	}
+	Centroid = Centroid * (1.0f / (Bone->Children.Num() + 1.0f));
+
+	Position = InputPosition = Centroid;
+	Rotation = InitialRotation = Bone->Rotation;
+	BoneLocalPosition = Bone->Rotation.Inverse() * (Bone->Position - Centroid);
+
+	ChildLocalPositions.Reserve(Bone->Children.Num());
+	for (const FBone* Child : Bone->Children)
+	{
+		FVector ChildLocalPos = Rotation.Inverse() * (Child->Position - Centroid);
+		ChildLocalPositions.Add(ChildLocalPos);
+	}
+}
+
 int FRigidBody::GetNumBonesToRoot() const
 { 
 	return NumBonesToRoot; 
diff --git a/Engine/Plugins/Experimental/FullBodyIK/Source/PBIK/Private/Core/PBIKConstraint.cpp b/Engine/Plugins/Experimental/FullBodyIK/Source/PBIK/Private/Core/PBIKConstraint.cpp
--- a/Engine/Plugins/Experimental/FullBodyIK/Source/PBIK/Private/Core/PBIKConstraint.cpp
+++ b/Engine/Plugins/Experimental/FullBodyIK/Source/PBIK/Private/Core/PBIKConstraint.cpp
@@ -72,6 +72,13 @@ void FJointConstraint::RemoveStretch(const float Percent)
 	B->Position -= Correction * Percent;
 }
 
+void FJointConstraint::UpdateFromInputs()
+{
+	const FVector PinPoint = B->Bone->Position;
+	PinPointLocalToA = A->Rotation.Inverse() * (PinPoint - A->Position);
+	PinPointLocalToB = B->Rotation.Inverse() * (PinPoint - B->Position);
+}
+
 FVector FJointConstraint::GetPositionCorrection(FVector& OutBodyToA, FVector& OutBodyToB) const
 {
 	OutBodyToA = A->Rotation * PinPointLocalToA;
diff --git a/Engine/Plugins/Experimental/FullBodyIK/Source/PBIK/Private/Core/PBIKSolver.cpp b/Engine/Plugins/Experimental/FullBodyIK/Source/PBIK/Private/Core/PBIKSolver.cpp
--- a/Engine/Plugins/Experimental/FullBodyIK/Source/PBIK/Private/Core/PBIKSolver.cpp
+++ b/Engine/Plugins/Experimental/FullBodyIK/Source/PBIK/Private/Core/PBIKSolver.cpp
@@ -168,6 +168,12 @@ void FPBIKSolver::Solve(const FPBIKSolverSettings& Settings)
 		Body.UpdateFromInputs(Settings);
 	}
 
+	// update constraints with new local pin locations relative to (possibly) translated joint positions
+	for (const TSharedPtr<PBIK::FConstraint>& Constraint : Constraints)
+	{
+		Constraint->UpdateFromInputs();
+	}
+
 	// optionally pin root in-place (convenience, does not require an effector)
 	if (RootPin.IsValid())
 	{
@@ -574,7 +580,7 @@ bool FPBIKSolver::InitBodies()
 	Bodies.Empty();
 	
 	// create bodies
-	for (FEffector& Effector : Effectors)
+	for (const FEffector& Effector : Effectors)
 	{
 		FBone* NextBone = Effector.Bone;
 		while (true)
diff --git a/Engine/Plugins/Experimental/FullBodyIK/Source/PBIK/Public/Core/PBIKBody.h b/Engine/Plugins/Experimental/FullBodyIK/Source/PBIK/Public/Core/PBIKBody.h
--- a/Engine/Plugins/Experimental/FullBodyIK/Source/PBIK/Public/Core/PBIKBody.h
+++ b/Engine/Plugins/Experimental/FullBodyIK/Source/PBIK/Public/Core/PBIKBody.h
@@ -100,10 +100,12 @@ public:
 
 	FRigidBody(FBone* InBone);
 
-	void Initialize(FBone* SolverRoot);
+	void Initialize(const FBone* SolverRoot);
 
 	void UpdateFromInputs(const FPBIKSolverSettings& Settings);
 
+	void UpdateTransformAndMassFromBones();
+
 	int GetNumBonesToRoot() const;
 
 	FRigidBody* GetParentBody() const;
diff --git a/Engine/Plugins/Experimental/FullBodyIK/Source/PBIK/Public/Core/PBIKConstraint.h b/Engine/Plugins/Experimental/FullBodyIK/Source/PBIK/Public/Core/PBIKConstraint.h
--- a/Engine/Plugins/Experimental/FullBodyIK/Source/PBIK/Public/Core/PBIKConstraint.h
+++ b/Engine/Plugins/Experimental/FullBodyIK/Source/PBIK/Public/Core/PBIKConstraint.h
@@ -16,6 +16,7 @@ struct FConstraint
 
 	virtual void Solve(const FPBIKSolverSettings& Settings) = 0;
 	virtual void RemoveStretch(const float Percent){};
+	virtual void UpdateFromInputs(){};
 };
 
 struct FJointConstraint : public FConstraint
@@ -58,6 +59,8 @@ public:
 
 	virtual void RemoveStretch(const float Percent) override;
 
+	virtual void UpdateFromInputs() override;
+
 private:
 
 	FVector GetPositionCorrection(FVector& OutBodyToA, FVector& OutBodyToB) const;
