GitHub 10186 : Fixed IK solving incorrectly if only one effector exists

PR #10186: Fixed IK solving incorrectly if only one effector exists (Contributed by [REDACTED_PERSON])

diff --git a/Engine/Plugins/Experimental/FullBodyIK/Source/PBIK/Private/Core/PBIKSolver.cpp b/Engine/Plugins/Experimental/FullBodyIK/Source/PBIK/Private/Core/PBIKSolver.cpp
--- a/Engine/Plugins/Experimental/FullBodyIK/Source/PBIK/Private/Core/PBIKSolver.cpp
+++ b/Engine/Plugins/Experimental/FullBodyIK/Source/PBIK/Private/Core/PBIKSolver.cpp
@@ -288,8 +288,8 @@ void FPBIKSolver::ApplyRootPrePull(const EPBIKRootBehavior RootBehavior, const F
 	}
 
 	// calculate "best fit" rotational delta from deformed effector configuration
-	FVector InitialCentroid;
-	FVector CurrentCentroid;
+	FVector InitialCentroid = FVector::ZeroVector;
+	FVector CurrentCentroid = FVector::ZeroVector;
 	const FQuat RotationOffset = GetRotationFromDeformedPoints(
 		InitialPoints,
 		CurrentPoints,
@@ -832,6 +832,8 @@ FQuat FPBIKSolver::GetRotationFromDeformedPoints(
 	// must have more than 1 point to generate a gradient
 	if (InInitialPoints.Num() <= 1)
 	{
+		OutInitialCentroid = (InInitialPoints.Num() == 1) ? InInitialPoints[0] : FVector::ZeroVector;
+		OutCurrentCentroid = (InCurrentPoints.Num() == 1) ? InCurrentPoints[0] : FVector::ZeroVector;
 		return FQuat::Identity;
 	}
 	
