Github fix: Fixed memory leak in FBIK solver (regression).

PR #10185: Fixed a memory leak in FullBodyIK plugin (Contributed by [REDACTED_PERSON])

diff --git a/Engine/Plugins/Experimental/FullBodyIK/Source/PBIK/Private/Core/PBIKBody.cpp b/Engine/Plugins/Experimental/FullBodyIK/Source/PBIK/Private/Core/PBIKBody.cpp
--- a/Engine/Plugins/Experimental/FullBodyIK/Source/PBIK/Private/Core/PBIKBody.cpp
+++ b/Engine/Plugins/Experimental/FullBodyIK/Source/PBIK/Private/Core/PBIKBody.cpp
@@ -99,7 +99,7 @@ void FRigidBody::UpdateTransformAndMassFromBones()
 	Rotation = InitialRotation = Bone->Rotation;
 	BoneLocalPosition = Bone->Rotation.Inverse() * (Bone->Position - Centroid);
 
-	ChildLocalPositions.Reserve(Bone->Children.Num());
+	ChildLocalPositions.Empty(Bone->Children.Num());
 	for (const FBone* Child : Bone->Children)
 	{
 		FVector ChildLocalPos = Rotation.Inverse() * (Child->Position - Centroid);
