Fixed loss of default values in FBIK solver in Control Rig.

#swarm https://p4-swarm.epicgames.net/reviews/24124771

diff --git a/Engine/Plugins/Experimental/FullBodyIK/Source/PBIK/Public/RigUnit_PBIK.h b/Engine/Plugins/Experimental/FullBodyIK/Source/PBIK/Public/RigUnit_PBIK.h
--- a/Engine/Plugins/Experimental/FullBodyIK/Source/PBIK/Public/RigUnit_PBIK.h
+++ b/Engine/Plugins/Experimental/FullBodyIK/Source/PBIK/Public/RigUnit_PBIK.h
@@ -55,7 +55,13 @@ struct FPBIKEffector
 {
 	GENERATED_BODY()
 
-	FPBIKEffector()	: Bone(NAME_None) {}
+	FPBIKEffector()	:
+		Bone(NAME_None),
+		PositionAlpha(1.0f),
+		RotationAlpha(1.0f),
+		StrengthAlpha(1.0f),
+		PullChainAlpha(1.0f),
+		PinRotation(1.0f){}
 
 	/** The bone that this effector will pull on. */
 	UPROPERTY(EditAnywhere, Category="Effector", meta = (CustomWidget = "BoneName"))
