Fixing FBIK solver to dynamically reinitialize with variable bone inputs.

diff --git a/Engine/Plugins/Experimental/FullBodyIK/Source/PBIK/Private/RigUnit_PBIK.cpp b/Engine/Plugins/Experimental/FullBodyIK/Source/PBIK/Private/RigUnit_PBIK.cpp
--- a/Engine/Plugins/Experimental/FullBodyIK/Source/PBIK/Private/RigUnit_PBIK.cpp
+++ b/Engine/Plugins/Experimental/FullBodyIK/Source/PBIK/Private/RigUnit_PBIK.cpp
@@ -5,6 +5,7 @@
 #include "Units/RigUnitContext.h"
 #include "Rigs/RigHierarchy.h"
 #include "Rigs/RigHierarchyElements.h"
+#include "Misc/HashBuilder.h"
 
 #include UE_INLINE_GENERATED_CPP_BY_NAME(RigUnit_PBIK)
 
@@ -13,22 +14,33 @@ FRigUnit_PBIK_Execute()
 	DECLARE_SCOPE_HIERARCHICAL_COUNTER_RIGUNIT()
 	LLM_SCOPE_BYNAME(TEXT("Animation/FBIK"));
 
-	if (bNeedsInit)
-	{
-		BoneSettingToSolverBoneIndex.Reset();
-		Solver.Reset();
-		SolverBoneToElementIndex.Reset();
-		EffectorSolverIndices.Reset();
-	}
-
 	URigHierarchy* Hierarchy = ExecuteContext.Hierarchy;
 	if (Hierarchy == nullptr)
 	{
 		return;
 	}
+	
+	// check if any inputs have been modified that would require reinitialization
+	FHashBuilder HashBuilder;
+	for (const FPBIKEffector& Effector : Effectors)
+	{
+		HashBuilder << Effector.Bone;
+	}
+	for (const FPBIKBoneSetting& BoneSetting : BoneSettings)
+	{
+		HashBuilder << BoneSetting.Bone;
+	}
+	uint32 Hash = HashBuilder.GetHash();
+	WorkData.bNeedsInit = WorkData.bNeedsInit ? true : Hash != WorkData.HashInitializedWith;
 
-	if (bNeedsInit)
+	// re-init if required
+	if (WorkData.bNeedsInit)
 	{
+		WorkData.BoneSettingToSolverBoneIndex.Reset();
+		WorkData.Solver.Reset();
+		WorkData.SolverBoneToElementIndex.Reset();
+		EffectorSolverIndices.Reset();
+		
 		// check how many effectors are assigned to a bone
 		int NumEffectors = 0;
 		for (const FPBIKEffector& Effector : Effectors)
@@ -51,7 +63,7 @@ FRigUnit_PBIK_Execute()
 		TArray<FRigBoneElement*> BoneElements = Hierarchy->GetBones(true);
 		for (int B = 0; B < BoneElements.Num(); ++B)
 		{
-			SolverBoneToElementIndex.Add(BoneElements[B]->GetIndex());
+			WorkData.SolverBoneToElementIndex.Add(BoneElements[B]->GetIndex());
 		}
 
 		// create bones
@@ -74,9 +86,9 @@ FRigUnit_PBIK_Execute()
 
 			// get the parent bone solver index
 			int ParentIndex = -1;
-			for (int P=0; P<SolverBoneToElementIndex.Num(); ++P)
+			for (int P=0; P<WorkData.SolverBoneToElementIndex.Num(); ++P)
 			{
-				if (SolverBoneToElementIndex[P] == ParentElementIndex)
+				if (WorkData.SolverBoneToElementIndex[P] == ParentElementIndex)
 				{
 					ParentIndex = P;
 					break;
@@ -87,22 +99,23 @@ FRigUnit_PBIK_Execute()
 			const FVector InOrigPosition = OrigTransform.GetLocation();
 			const FQuat InOrigRotation = OrigTransform.GetRotation();
 			bool bIsRoot = BoneElements[B]->GetName() == Root;
-			Solver.AddBone(Name, ParentIndex, InOrigPosition, InOrigRotation, bIsRoot);
+			WorkData.Solver.AddBone(Name, ParentIndex, InOrigPosition, InOrigRotation, bIsRoot);
 		}
 		
 		// create effectors
 		for (const FPBIKEffector& Effector : Effectors)
 		{
-			int32 IndexInSolver = Solver.AddEffector(Effector.Bone);
+			int32 IndexInSolver = WorkData.Solver.AddEffector(Effector.Bone);
 			EffectorSolverIndices.Add(IndexInSolver);
 		}
 		
 		// initialize
-		Solver.Initialize();
-		bNeedsInit = false;
+		WorkData.Solver.Initialize();
+		WorkData.bNeedsInit = false;
+		WorkData.HashInitializedWith = Hash;
 	}
 
-	if (!Solver.IsReadyToSimulate())
+	if (!WorkData.Solver.IsReadyToSimulate())
 	{
 		return;
 	}
@@ -113,19 +126,19 @@ FRigUnit_PBIK_Execute()
 	}
 
 	// set bones to input pose
-	for(int32 BoneIndex = 0; BoneIndex < Solver.GetNumBones(); BoneIndex++)
+	for(int32 BoneIndex = 0; BoneIndex < WorkData.Solver.GetNumBones(); BoneIndex++)
 	{
-		const FTransform GlobalTransform = Hierarchy->GetGlobalTransform(SolverBoneToElementIndex[BoneIndex]);
-		Solver.SetBoneTransform(BoneIndex, GlobalTransform);
+		const FTransform GlobalTransform = Hierarchy->GetGlobalTransform(WorkData.SolverBoneToElementIndex[BoneIndex]);
+		WorkData.Solver.SetBoneTransform(BoneIndex, GlobalTransform);
 	}
 
 	// invalidate the name lookup for the settings
-	if(BoneSettingToSolverBoneIndex.Num() != BoneSettings.Num())
+	if(WorkData.BoneSettingToSolverBoneIndex.Num() != BoneSettings.Num())
 	{
-		BoneSettingToSolverBoneIndex.Reset();
-		while(BoneSettingToSolverBoneIndex.Num() < BoneSettings.Num())
+		WorkData.BoneSettingToSolverBoneIndex.Reset();
+		while(WorkData.BoneSettingToSolverBoneIndex.Num() < BoneSettings.Num())
 		{
-			BoneSettingToSolverBoneIndex.Add(INDEX_NONE);
+			WorkData.BoneSettingToSolverBoneIndex.Add(INDEX_NONE);
 		}
 	}
 
@@ -134,17 +147,17 @@ FRigUnit_PBIK_Execute()
 	{
 		const FPBIKBoneSetting& BoneSetting = BoneSettings[BoneSettingIndex];
 
-		if(BoneSettingToSolverBoneIndex[BoneSettingIndex] == INDEX_NONE)
+		if(WorkData.BoneSettingToSolverBoneIndex[BoneSettingIndex] == INDEX_NONE)
 		{
-			BoneSettingToSolverBoneIndex[BoneSettingIndex] = Solver.GetBoneIndex(BoneSetting.Bone);
-			if(BoneSettingToSolverBoneIndex[BoneSettingIndex] == INDEX_NONE)
+			WorkData.BoneSettingToSolverBoneIndex[BoneSettingIndex] = WorkData.Solver.GetBoneIndex(BoneSetting.Bone);
+			if(WorkData.BoneSettingToSolverBoneIndex[BoneSettingIndex] == INDEX_NONE)
 			{
 				continue;
 			}
 		}
 
-		int32 BoneIndex = BoneSettingToSolverBoneIndex[BoneSettingIndex];
-		if (PBIK::FBoneSettings* InternalSettings = Solver.GetBoneSettings(BoneIndex))
+		int32 BoneIndex = WorkData.BoneSettingToSolverBoneIndex[BoneSettingIndex];
+		if (PBIK::FBoneSettings* InternalSettings = WorkData.Solver.GetBoneSettings(BoneIndex))
 		{
 			BoneSetting.CopyToCoreStruct(*InternalSettings);
 		}
@@ -169,22 +182,22 @@ FRigUnit_PBIK_Execute()
 		EffectorSettings.PullChainAlpha = Effector.PullChainAlpha;
 		EffectorSettings.PinRotation = Effector.PinRotation;
 		
-		Solver.SetEffectorGoal(EffectorSolverIndices[E], Position, Rotation, EffectorSettings);
+		WorkData.Solver.SetEffectorGoal(EffectorSolverIndices[E], Position, Rotation, EffectorSettings);
 	}
 
 	// solve
-	Solver.Solve(Settings);
+	WorkData.Solver.Solve(Settings);
 
 	// copy transforms back
 	const bool bPropagateTransform = false;
-	for(int32 BoneIndex = 0; BoneIndex < Solver.GetNumBones(); BoneIndex++)
+	for(int32 BoneIndex = 0; BoneIndex < WorkData.Solver.GetNumBones(); BoneIndex++)
 	{
 		FTransform NewTransform;
-		Solver.GetBoneGlobalTransform(BoneIndex, NewTransform);
-		Hierarchy->SetGlobalTransform(SolverBoneToElementIndex[BoneIndex], NewTransform, false, bPropagateTransform);
+		WorkData.Solver.GetBoneGlobalTransform(BoneIndex, NewTransform);
+		Hierarchy->SetGlobalTransform(WorkData.SolverBoneToElementIndex[BoneIndex], NewTransform, false, bPropagateTransform);
 	}
 
 	// do all debug drawing
-	Debug.Draw(ExecuteContext.GetDrawInterface(), &Solver);
+	Debug.Draw(ExecuteContext.GetDrawInterface(), &WorkData.Solver);
 }
 
diff --git a/Engine/Plugins/Experimental/FullBodyIK/Source/PBIK/Public/PBIK_Shared.h b/Engine/Plugins/Experimental/FullBodyIK/Source/PBIK/Public/PBIK_Shared.h
--- a/Engine/Plugins/Experimental/FullBodyIK/Source/PBIK/Public/PBIK_Shared.h
+++ b/Engine/Plugins/Experimental/FullBodyIK/Source/PBIK/Public/PBIK_Shared.h
@@ -28,7 +28,7 @@ struct PBIK_API FPBIKBoneSetting
 		PreferredAngles(FVector::ZeroVector){}
 
 	/** The Bone that these settings will be applied to. */
-	UPROPERTY(EditAnywhere, Category = Bone, meta = (Constant, CustomWidget = "BoneName"))
+	UPROPERTY(EditAnywhere, Category = Bone, meta = (CustomWidget = "BoneName"))
 	FName Bone;
 
 	/** Range is 0 to 1 (Default is 0). At higher values, the bone will resist rotating (forcing other bones to compensate). */
diff --git a/Engine/Plugins/Experimental/FullBodyIK/Source/PBIK/Public/RigUnit_PBIK.h b/Engine/Plugins/Experimental/FullBodyIK/Source/PBIK/Public/RigUnit_PBIK.h
--- a/Engine/Plugins/Experimental/FullBodyIK/Source/PBIK/Public/RigUnit_PBIK.h
+++ b/Engine/Plugins/Experimental/FullBodyIK/Source/PBIK/Public/RigUnit_PBIK.h
@@ -100,10 +100,37 @@ struct FPBIKEffector
 	float PinRotation = 1.0f;
 };
 
+USTRUCT()
+struct FPBIKWorkData
+{
+	GENERATED_BODY()
+	
+	FPBIKWorkData() :
+		bNeedsInit(true),
+		HashInitializedWith(-1)
+	{
+	}
+	
+	UPROPERTY(transient)
+	bool bNeedsInit;
+
+	UPROPERTY(transient)
+	uint32 HashInitializedWith;
+	
+	UPROPERTY(transient)
+	TArray<int32> BoneSettingToSolverBoneIndex;
+
+	UPROPERTY(transient)
+	TArray<int32> SolverBoneToElementIndex;
+	
+	UPROPERTY(transient)
+	FPBIKSolver Solver;
+};
+
 /*
  * Based on a Position Based solver at core, this node can solve multi chains within a root using multi effectors
  */
-USTRUCT(meta=(DisplayName="Full Body IK", Category="Hierarchy", Keywords="Position Based, PBIK, IK, Full Body, Multi, Effector, N-Chain, FB, HIK, HumanIK", NodeColor="0 1 1"))
+USTRUCT(meta=(DisplayName="Full Body IK", Category="Hierarchy", Keywords="FBIK, Position Based, PBIK, IK, Full Body, Multi, Effector, N-Chain, FB, HIK, HumanIK", NodeColor="0 1 1"))
 struct FRigUnit_PBIK : public FRigUnit_HighlevelBaseMutable
 {
 	GENERATED_BODY()
@@ -111,11 +138,7 @@ struct FRigUnit_PBIK : public FRigUnit_HighlevelBaseMutable
 	RIGVM_METHOD()
 	virtual void Execute() override;
 
-	FRigUnit_PBIK() :
-		Root(NAME_None),
-		bNeedsInit(true)
-	{
-	}
+	FRigUnit_PBIK() : Root(NAME_None) { }
 
 	/**This is usually the top-most skinned bone; often the "Pelvis" or "Hips", but can be set to any bone.
 	 *Bones above the root will be ignored by the solver.
@@ -147,15 +170,7 @@ struct FRigUnit_PBIK : public FRigUnit_HighlevelBaseMutable
 	UPROPERTY(meta = (Input))
 	FPBIKDebug Debug;
 
+	/** Runtime-only data */
 	UPROPERTY(transient)
-	TArray<int32> BoneSettingToSolverBoneIndex;
-
-	UPROPERTY(transient)
-	TArray<int32> SolverBoneToElementIndex;
-
-	UPROPERTY(transient)
-	FPBIKSolver Solver;
-
-	UPROPERTY(transient)
-	bool bNeedsInit;
+	FPBIKWorkData WorkData;
 };
