Fixes NANs in Jacobian FBIK solver.

diff --git a/Engine/Plugins/Experimental/FullBodyIK/Source/FullBodyIK/Private/FBIKUtil.cpp b/Engine/Plugins/Experimental/FullBodyIK/Source/FullBodyIK/Private/FBIKUtil.cpp
--- a/Engine/Plugins/Experimental/FullBodyIK/Source/FullBodyIK/Private/FBIKUtil.cpp
+++ b/Engine/Plugins/Experimental/FullBodyIK/Source/FullBodyIK/Private/FBIKUtil.cpp
@@ -9,6 +9,11 @@ namespace FBIKUtil
 {
 	bool CanCrossProduct(const FVector& Vector1, const FVector& Vector2)
 	{
+		// zero vectors produce NANs below
+		if (Vector1.IsNearlyZero() || Vector2.IsNearlyZero())
+		{
+			return false;
+		}
 		// here we test their parallelism, and we normalize vectors
 		float Cosine = FVector::DotProduct(Vector1.GetUnsafeNormal(), Vector2.GetUnsafeNormal());
 		// if they're parallel, we need to find other axis
