Fixed bug, preferred angles were applied additively on incoming animation in FBIK solver.

diff --git a/Engine/Plugins/Experimental/FullBodyIK/Source/PBIK/Private/Core/PBIKBody.cpp b/Engine/Plugins/Experimental/FullBodyIK/Source/PBIK/Private/Core/PBIKBody.cpp
--- a/Engine/Plugins/Experimental/FullBodyIK/Source/PBIK/Private/Core/PBIKBody.cpp
+++ b/Engine/Plugins/Experimental/FullBodyIK/Source/PBIK/Private/Core/PBIKBody.cpp
@@ -41,9 +41,9 @@ void FBone::UpdateFromInputs()
 		return;
 	}
 
-	LocalPositionOrig = Parent->Rotation.Inverse() * (Position - Parent->Position);
-	LocalRotationOrig = Parent->Rotation.Inverse() * Rotation;
-	Length = LocalPositionOrig.Size();
+	LocalPositionFromInput = Parent->Rotation.Inverse() * (Position - Parent->Position);
+	LocalRotationFromInput = Parent->Rotation.Inverse() * Rotation;
+	Length = LocalPositionFromInput.Size();
 }
 
 FRigidBody::FRigidBody(FBone* InBone)
diff --git a/Engine/Plugins/Experimental/FullBodyIK/Source/PBIK/Private/Core/PBIKSolver.cpp b/Engine/Plugins/Experimental/FullBodyIK/Source/PBIK/Private/Core/PBIKSolver.cpp
--- a/Engine/Plugins/Experimental/FullBodyIK/Source/PBIK/Private/Core/PBIKSolver.cpp
+++ b/Engine/Plugins/Experimental/FullBodyIK/Source/PBIK/Private/Core/PBIKSolver.cpp
@@ -39,31 +39,7 @@ namespace PBIK
 		Pin.Pin()->SetGoal(Position, Rotation, Settings.StrengthAlpha);
 
 		// update length of chain this effector controls in the input pose
-		if (ParentSubRoot)
-		{
-			FEffector* ParentEffector = ParentSubRoot->Effector;
-			const FVector ParentSubRootPosition = ParentEffector ? ParentEffector->PositionOrig : ParentSubRoot->Position;
-			DistanceToSubRootInInputPose = (ParentSubRootPosition - Bone->Position).Size();
-		}
-
-		DistToRootStraightLine = (PositionOrig - SolverRoot.Position).Size();
-
-		// update distances to root (along bones)
-		{
-			if (Bone->bIsSolverRoot)
-			{
-				DistToRootAlongBones = 0.0f;
-			}else
-			{
-				DistToRootAlongBones = Bone->Length;
-				FBone* Parent = Bone->Parent;
-				while (Parent && !Parent->bIsSolverRoot)
-				{
-					DistToRootAlongBones += Parent->Length;
-					Parent = Parent->Parent;
-				}
-			}
-		}
+		DistToSubRootInInputPose = CalculateDistanceToSubRoot();
 
 		// update distances to SUB root (along bones)
 		if (ParentSubRoot)
@@ -71,7 +47,7 @@ namespace PBIK
 			DistancesFromEffector.Reset();
 			DistancesFromEffector.Add(0.0f);
 			DistToSubRootAlongBones = Bone->Length;
-			FBone* Parent = Bone->Parent;
+			const FBone* Parent = Bone->Parent;
 			while (Parent && Parent->Body != ParentSubRoot)
 			{
 				DistancesFromEffector.Add(DistToSubRootAlongBones);
@@ -81,7 +57,20 @@ namespace PBIK
 		}
 	}
 
-	void FEffector::ApplyPreferredAngles()
+	float FEffector::CalculateDistanceToSubRoot() const
+	{
+		// calculates distance from the bone on this effector to the parent sub-root (chain root)
+		if (ParentSubRoot)
+		{
+			const FEffector* ParentEffector = ParentSubRoot->Effector;
+			const FVector ParentSubRootPosition = ParentEffector ? ParentEffector->PositionOrig : ParentSubRoot->Position;
+			return (ParentSubRootPosition - Bone->Position).Size();
+		}
+
+		return 0.f;
+	}
+
+	void FEffector::ApplyPreferredAngles() const
 	{
 		// optionally apply a preferred angle to give solver a hint which direction to favor
 		// apply amount of preferred angle proportional to the amount this sub-limb is squashed
@@ -93,7 +82,7 @@ namespace PBIK
 		}
 
 		// can't squash chain with zero length already
-		if (DistanceToSubRootInInputPose <= SMALL_NUMBER)
+		if (DistToSubRootInInputPose <= SMALL_NUMBER || DistToSubRootInRefPose <= SMALL_NUMBER)
 		{
 			return;
 		}
@@ -101,32 +90,60 @@ namespace PBIK
 		// we have to be careful here when calculating the distance to the parent sub-root.
 		// if the parent sub-root is attached to an effector, use the effector's position
 		// otherwise use the current position of the FRigidBody
-		FEffector* ParentEffector = ParentSubRoot->Effector;
+		const FEffector* ParentEffector = ParentSubRoot->Effector;
 		const FVector ParentSubRootPosition = ParentEffector ? ParentEffector->Position : ParentSubRoot->Position;
-		const float DistToParentSubRoot = (ParentSubRootPosition - Position).Size();
-		if (DistToParentSubRoot >= DistanceToSubRootInInputPose)
+		const float DistToSubRootCurrent = (ParentSubRootPosition - Position).Size();
+		if (DistToSubRootCurrent >= DistToSubRootInInputPose)
 		{
 			return; // limb is stretched
 		}
-		
+
 		// amount squashed (clamped to scaled original length)
-		float DeltaSquash = DistanceToSubRootInInputPose - DistToParentSubRoot;
-		DeltaSquash = DeltaSquash > DistanceToSubRootInInputPose ? DistanceToSubRootInInputPose : DeltaSquash;
-		float SquashPercent = DeltaSquash / DistanceToSubRootInInputPose;
+		const float DeltaSquash = DistToSubRootInRefPose - DistToSubRootCurrent;
+		float SquashPercent = DeltaSquash / DistToSubRootInRefPose;
 		SquashPercent = PBIK::CircularEaseOut(SquashPercent);
 		if (FMath::IsNearlyZero(SquashPercent))
 		{
 			return; // limb not squashed enough
 		}
 
-		FBone* Parent = Bone->Parent;
+		// iterate over all the bodies in the chain controlled by this effector...
+		const FBone* Parent = Bone->Parent;
 		while (Parent && Parent->bIsSolved)
 		{
-			if (Parent->Body->J.bUsePreferredAngles)
+			FRigidBody* Body = Parent->Body;
+			if (Body->J.bUsePreferredAngles)
 			{
-				FQuat PartialRotation = FQuat::FastLerp(FQuat::Identity, FQuat(Parent->Body->J.PreferredAngles), SquashPercent);
-				Parent->Body->Rotation = Parent->Body->Rotation * PartialRotation;
-				Parent->Body->Rotation.Normalize();
+				const FRigidBody* CurrentParent = Body->GetParentBody();
+				FQuat LocalRotation = CurrentParent ? CurrentParent->Rotation.Inverse() * Body->Rotation : Body->Rotation;
+				const FRotator AnglesFromInput = LocalRotation.Rotator();
+				const FRotator AnglesFromRefPose = Body->Bone->LocalRotationInitial;
+				const FRotator PreferredAngles = Body->J.PreferredAngles * SquashPercent;
+
+				auto CalcDeltaAngle = [](double PreferredAngle, double RefPoseAngle, double InputAngle) -> double
+				{
+					if (FMath::Abs(PreferredAngle) <= SMALL_NUMBER)
+					{
+						return 0.f; // apply no delta, preferred angle was unset at 0
+					}
+
+					// if the target angle is greater than the input pose, return a delta
+					const float TargetAngle = PreferredAngle + RefPoseAngle;
+					if (FMath::Abs(TargetAngle) > FMath::Abs(InputAngle))
+					{
+						return TargetAngle - InputAngle;
+					}
+
+					// body is already bent more than the preferred angle
+					return 0.f;
+				};
+
+				FRotator LocalDeltaAngles;
+				LocalDeltaAngles.Pitch = CalcDeltaAngle(PreferredAngles.Pitch, AnglesFromRefPose.Pitch, AnglesFromInput.Pitch);
+				LocalDeltaAngles.Yaw = CalcDeltaAngle(PreferredAngles.Yaw, AnglesFromRefPose.Yaw, AnglesFromInput.Yaw);
+				LocalDeltaAngles.Roll = CalcDeltaAngle(PreferredAngles.Roll, AnglesFromRefPose.Roll, AnglesFromInput.Roll);
+				Body->Rotation = Body->Rotation * LocalDeltaAngles.Quaternion();
+				Body->Rotation.Normalize();
 			}
 
 			if (Parent == ParentSubRoot->Bone)
@@ -240,7 +257,7 @@ void FPBIKSolver::UpdateBonesFromBodies()
 			continue; // effector is between other effectors, so leave transform where body ended up
 		}
 		
-		Bone->Position = Bone->Parent->Position + Bone->Parent->Rotation * Bone->LocalPositionOrig;
+		Bone->Position = Bone->Parent->Position + Bone->Parent->Rotation * Bone->LocalPositionFromInput;
 
 		if (Effector.Settings.PinRotation > SMALL_NUMBER)
 		{
@@ -249,7 +266,7 @@ void FPBIKSolver::UpdateBonesFromBodies()
 			Bone->Rotation = FQuat::FastLerp(Bone->Rotation, Effector.Rotation, RotAmount).GetNormalized();
 		}else
 		{
-			Bone->Rotation = Bone->Parent->Rotation * Bone->LocalRotationOrig;
+			Bone->Rotation = Bone->Parent->Rotation * Bone->LocalRotationFromInput;
 		}
 	}
 
@@ -261,8 +278,8 @@ void FPBIKSolver::UpdateBonesFromBodies()
 			continue;
 		}
 
-		Bone.Position = Bone.Parent->Position + Bone.Parent->Rotation * Bone.LocalPositionOrig;
-		Bone.Rotation = Bone.Parent->Rotation * Bone.LocalRotationOrig;
+		Bone.Position = Bone.Parent->Position + Bone.Parent->Rotation * Bone.LocalPositionFromInput;
+		Bone.Rotation = Bone.Parent->Rotation * Bone.LocalRotationFromInput;
 	}
 }
 
@@ -363,7 +380,7 @@ void FPBIKSolver::ApplyPullChainAlpha(const float GlobalPullChainAlpha)
 		if (Effector.ParentSubRoot->Effector)
 		{
 			// add effector offset to chain root
-			FEffector* RootEffector = Effector.ParentSubRoot->Effector;
+			const FEffector* RootEffector = Effector.ParentSubRoot->Effector;
 			FVector RootEffectorDelta = RootEffector->Position - RootEffector->PositionOrig;
 			ChainStartNew = Effector.ParentSubRoot->Position + RootEffectorDelta;
 		}else
@@ -380,7 +397,7 @@ void FPBIKSolver::ApplyPullChainAlpha(const float GlobalPullChainAlpha)
 		ChainDeltaRotation = FQuat::FastLerp(FQuat::Identity, ChainDeltaRotation, ChainDeltaAlpha);
 		FVector ChainDeltaPosition = ChainVecNew * (ChainLenNew - ChainLenOrig) * ChainDeltaAlpha;
 
-		FBone* Bone = Effector.Bone->Parent;
+		const FBone* Bone = Effector.Bone->Parent;
 		int32 ChainIndex = 0;
 		const float InvChainLength = 1.0f / Effector.DistToSubRootAlongBones;
 		while (Bone && Bone->Body != Effector.ParentSubRoot)
@@ -463,6 +480,11 @@ bool FPBIKSolver::Initialize()
 		return false;
 	}
 
+	if (!InitEffectors())
+    {
+    	return false;
+    }
+
 	if (!InitConstraints())
 	{
 		return false;
@@ -535,7 +557,7 @@ bool FPBIKSolver::InitBones()
 	}
 
 	// walk upwards from each effector to the root to initialize "Bone.IsSolved"
-	for (FEffector& Effector : Effectors)
+	for (const FEffector& Effector : Effectors)
 	{
 		FBone* NextBone = Effector.Bone;
 		while (NextBone)
@@ -568,6 +590,13 @@ bool FPBIKSolver::InitBones()
 		Bone.bIsSubRoot = Bone.Children.Num() > 1 || Bone.bIsSolverRoot;
 	}
 
+	// store initial local rotation (preferred angles are relative to this)
+	for (FBone& Bone : Bones)
+	{
+		const FQuat LocalRotation = Bone.Parent ? Bone.Parent->Rotation.Inverse() * Bone.Rotation : Bone.Rotation;
+		Bone.LocalRotationInitial = LocalRotation.Rotator();
+	}
+
 	return true;
 }
 
@@ -643,9 +672,23 @@ bool FPBIKSolver::InitBodies()
 	return true;
 }
 
+bool FPBIKSolver::InitEffectors()
+{
+	using PBIK::FEffector;
+
+	// store the length of chain this effector controls in the reference pose
+	for (FEffector& Effector : Effectors)
+	{
+		Effector.DistToSubRootInRefPose = Effector.CalculateDistanceToSubRoot();
+	}
+
+	return true;
+}
+
 void FPBIKSolver::AddBodyForBone(PBIK::FBone* Bone)
 {
-	for (PBIK::FRigidBody& Body : Bodies)
+	using PBIK::FRigidBody;
+	for (const FRigidBody& Body : Bodies)
 	{
 		if (Body.Bone == Bone)
 		{
diff --git a/Engine/Plugins/Experimental/FullBodyIK/Source/PBIK/Public/Core/PBIKBody.h b/Engine/Plugins/Experimental/FullBodyIK/Source/PBIK/Public/Core/PBIKBody.h
--- a/Engine/Plugins/Experimental/FullBodyIK/Source/PBIK/Public/Core/PBIKBody.h
+++ b/Engine/Plugins/Experimental/FullBodyIK/Source/PBIK/Public/Core/PBIKBody.h
@@ -20,8 +20,9 @@ struct FBone
 	FVector Position;
 	FQuat Rotation;
 	FVector Scale; // just passed through, not modified
-	FVector LocalPositionOrig;
-	FQuat LocalRotationOrig;
+	FVector LocalPositionFromInput;
+	FQuat LocalRotationFromInput;
+	FRotator LocalRotationInitial;
 
 	// initialized - these fields are null/empty until after Solver::Initialize()
 	FRigidBody* Body = nullptr;
diff --git a/Engine/Plugins/Experimental/FullBodyIK/Source/PBIK/Public/Core/PBIKSolver.h b/Engine/Plugins/Experimental/FullBodyIK/Source/PBIK/Public/Core/PBIKSolver.h
--- a/Engine/Plugins/Experimental/FullBodyIK/Source/PBIK/Public/Core/PBIKSolver.h
+++ b/Engine/Plugins/Experimental/FullBodyIK/Source/PBIK/Public/Core/PBIKSolver.h
@@ -75,9 +75,8 @@ struct FEffector
 	FBone* Bone;
 	TWeakPtr<FPinConstraint> Pin;
 	FRigidBody* ParentSubRoot = nullptr;
-	float DistanceToSubRootInInputPose;
-	float DistToRootAlongBones;
-	float DistToRootStraightLine;
+	float DistToSubRootInInputPose;
+	float DistToSubRootInRefPose;
 	
 	TArray<float> DistancesFromEffector;
 	float DistToSubRootAlongBones;
@@ -90,7 +89,8 @@ struct FEffector
 		const FEffectorSettings& InSettings);
 
 	void UpdateFromInputs(const FBone& SolverRoot);
-	void ApplyPreferredAngles();
+	float CalculateDistanceToSubRoot() const;
+	void ApplyPreferredAngles() const;
 };
 	
 } // namespace
@@ -248,6 +248,8 @@ private:
 
 	bool InitBodies();
 
+	bool InitEffectors();
+
 	bool InitConstraints();
 
 	void AddBodyForBone(PBIK::FBone* Bone);
