Added "Chain Depth" to enable better preferred angle blending in FBIK solver. This also preparation for "Sub Solves" to be added soon.

diff --git a/Engine/Plugins/Animation/IKRig/Source/IKRig/Private/Rig/Solvers/IKRig_PBIKSolver.cpp b/Engine/Plugins/Animation/IKRig/Source/IKRig/Private/Rig/Solvers/IKRig_PBIKSolver.cpp
--- a/Engine/Plugins/Animation/IKRig/Source/IKRig/Private/Rig/Solvers/IKRig_PBIKSolver.cpp
+++ b/Engine/Plugins/Animation/IKRig/Source/IKRig/Private/Rig/Solvers/IKRig_PBIKSolver.cpp
@@ -104,6 +104,7 @@ void UIKRigFBIKSolver::Solve(FIKRigSkeleton& IKRigSkeleton, const FIKRigGoalCont
 		Settings.PositionAlpha = 1.0f; // this is constant because IKRig manages offset alphas itself
 		Settings.RotationAlpha = 1.0f; // this is constant because IKRig manages offset alphas itself
 		Settings.StrengthAlpha = Effector->StrengthAlpha;
+		Settings.ChainDepth = Effector->ChainDepth;
 		Settings.PullChainAlpha = Effector->PullChainAlpha;
 		Settings.PinRotation = Effector->PinRotation;
 		
diff --git a/Engine/Plugins/Animation/IKRig/Source/IKRig/Public/Rig/Solvers/IKRig_FBIKSolver.h b/Engine/Plugins/Animation/IKRig/Source/IKRig/Public/Rig/Solvers/IKRig_FBIKSolver.h
--- a/Engine/Plugins/Animation/IKRig/Source/IKRig/Public/Rig/Solvers/IKRig_FBIKSolver.h
+++ b/Engine/Plugins/Animation/IKRig/Source/IKRig/Public/Rig/Solvers/IKRig_FBIKSolver.h
@@ -25,6 +25,13 @@ public:
 	UPROPERTY(BlueprintReadOnly, VisibleAnywhere, Category = "Full Body IK Effector")
 	FName BoneName;
 
+	/** Range 0-inf (default is 0). Explicitly set the number of bones up the hierarchy to consider part of this effector's 'chain'.
+	* The "chain" of bones is used to apply Preferred Angles, Pull Chain Alpha and Chain "Sub Solves".
+	* If left at 0, the solver will attempt to determine the root of the chain by searching up the hierarchy until it finds a branch or another effector, whichever it finds first.
+	*/
+	UPROPERTY(BlueprintReadWrite, EditAnywhere, Category = "Full Body IK Effector", meta = (ClampMin = "0", UIMin = "0"))
+	int32 ChainDepth = 0;
+
 	/** Range 0-1 (default is 1.0). The strength of the effector when pulling the bone towards it's target location.
 	* At 0.0, the effector does not pull at all, but the bones between the effector and the root will still slightly resist motion from other effectors.
 	* This can thus act as a "stabilizer" for parts of the body that you do not want to behave in a pure FK fashion.
@@ -32,10 +39,11 @@ public:
 	UPROPERTY(BlueprintReadWrite, EditAnywhere, Category = "Full Body IK Effector", meta = (ClampMin = "0", ClampMax = "1", UIMin = "0.0", UIMax = "1.0"))
 	float StrengthAlpha = 1.0f;
 
-	/** Range 0-1 (default is 1.0). When enabled (greater than 0.0), the solver internally partitions the skeleton into 'chains' which extend from the effector to the nearest fork in the skeleton.
-	*These chains are pre-rotated and translated, as a whole, towards the effector targets.
-	*This can improve the results for sparse bone chains, and significantly improve convergence on dense bone chains.
-	*But it may cause undesirable results in highly constrained bone chains (like robot arms).
+	/** Range 0-1 (default is 1.0). When enabled (greater than 0.0), the solver internally partitions the skeleton into 'chains' which extend
+	 * from the effector up the hierarchy by "Chain Depth". If Chain Depth is 0, the chain root is set to the nearest fork in the skeleton.
+	* These chains are pre-rotated and translated, as a whole, towards the effector targets.
+	* This can improve the results for sparse bone chains, and significantly improve convergence on dense bone chains.
+	* But it may cause undesirable results in highly constrained bone chains (like robot arms).
 	*/
 	UPROPERTY(BlueprintReadWrite, EditAnywhere, Category = "Full Body IK Effector", meta = (ClampMin = "0", ClampMax = "1", UIMin = "0.0", UIMax = "1.0"))
 	float PullChainAlpha = 1.0f;
@@ -50,6 +58,7 @@ public:
 
 	void CopySettings(const UIKRig_FBIKEffector* Other)
 	{
+		ChainDepth = Other->ChainDepth;
 		StrengthAlpha = Other->StrengthAlpha;
 		PullChainAlpha = Other->PullChainAlpha;
 		PinRotation = Other->PinRotation;
diff --git a/Engine/Plugins/Experimental/FullBodyIK/Source/PBIK/Private/Core/PBIKSolver.cpp b/Engine/Plugins/Experimental/FullBodyIK/Source/PBIK/Private/Core/PBIKSolver.cpp
--- a/Engine/Plugins/Experimental/FullBodyIK/Source/PBIK/Private/Core/PBIKSolver.cpp
+++ b/Engine/Plugins/Experimental/FullBodyIK/Source/PBIK/Private/Core/PBIKSolver.cpp
@@ -31,6 +31,47 @@ namespace PBIK
 		Settings = InSettings;
 	}
 
+	void FEffector::UpdateChainRoot()
+	{
+		// NOTE: this function must be called AFTER InitBodies() due to it's reliance on the FBody.bIsSubRoot flag
+		
+		ChainRootBody = nullptr;
+		const int32 MaxDepth = Settings.ChainDepth;
+		const FBone* Parent = Bone->Parent;
+		int32 Depth = 1;
+		while (Parent)
+		{
+			if (!Parent->bIsSolved)
+			{
+				break; // this only happens when effector is on solver root
+			}
+
+			if (MaxDepth > 0)
+			{
+				// user specified a custom chain depth, so use it
+				if (Depth >= MaxDepth)
+				{
+					ChainRootBody = Parent->Body;
+					break;
+				}
+			}
+			else
+			{
+				// user did not specify a max depth, so use the nearest sub-root as the chain root
+				if (Parent->bIsSubRoot || Parent->bIsSolverRoot)
+				{
+					ChainRootBody = Parent->Body;
+					break;
+				}
+			}
+
+			++Depth;
+			Parent = Parent->Parent;
+		}
+
+		ChainRootDepthInitializedWith = Settings.ChainDepth;
+	}
+
 	void FEffector::UpdateFromInputs(const FBone& SolverRoot)
 	{
 		// blend effector transform by alpha and set pin goal transform
@@ -38,32 +79,37 @@ namespace PBIK
 		Rotation = FMath::Lerp(RotationOrig, RotationGoal, Settings.RotationAlpha);
 		Pin.Pin()->SetGoal(Position, Rotation, Settings.StrengthAlpha);
 
+		// in case user modified the chain depth at runtime, we may need to find the new chain root 
+		if (Settings.ChainDepth != ChainRootDepthInitializedWith)
+		{
+			UpdateChainRoot();
+		}
 		// update length of chain this effector controls in the input pose
-		DistToSubRootInInputPose = CalculateDistanceToSubRoot();
+		DistToChainRootInInputPose = CalculateDistanceToChainRoot();
 
-		// update distances to SUB root (along bones)
-		if (ParentSubRoot)
+		// update distances to chain root (along bones)
+		if (ChainRootBody)
 		{
 			DistancesFromEffector.Reset();
 			DistancesFromEffector.Add(0.0f);
-			DistToSubRootAlongBones = Bone->Length;
+			DistToChainRootAlongBones = Bone->Length;
 			const FBone* Parent = Bone->Parent;
-			while (Parent && Parent->Body != ParentSubRoot)
+			while (Parent && Parent->Body != ChainRootBody)
 			{
-				DistancesFromEffector.Add(DistToSubRootAlongBones);
-				DistToSubRootAlongBones += Parent->Length;
+				DistancesFromEffector.Add(DistToChainRootAlongBones);
+				DistToChainRootAlongBones += Parent->Length;
 				Parent = Parent->Parent;
 			}
 		}
 	}
 
-	float FEffector::CalculateDistanceToSubRoot() const
+	float FEffector::CalculateDistanceToChainRoot() const
 	{
 		// calculates distance from the bone on this effector to the parent sub-root (chain root)
-		if (ParentSubRoot)
+		if (ChainRootBody)
 		{
-			const FEffector* ParentEffector = ParentSubRoot->Effector;
-			const FVector ParentSubRootPosition = ParentEffector ? ParentEffector->PositionOrig : ParentSubRoot->Position;
+			const FEffector* ParentEffector = ChainRootBody->Effector;
+			const FVector ParentSubRootPosition = ParentEffector ? ParentEffector->PositionOrig : ChainRootBody->Position;
 			return (ParentSubRootPosition - Bone->Position).Size();
 		}
 
@@ -76,13 +122,13 @@ namespace PBIK
 		// apply amount of preferred angle proportional to the amount this sub-limb is squashed
 
 		// can't squash root chain
-		if (!ParentSubRoot)
+		if (!ChainRootBody)
 		{
 			return; 
 		}
 
 		// can't squash chain with zero length already
-		if (DistToSubRootInInputPose <= SMALL_NUMBER || DistToSubRootInRefPose <= SMALL_NUMBER)
+		if (DistToChainRootInInputPose <= SMALL_NUMBER)
 		{
 			return;
 		}
@@ -90,17 +136,17 @@ namespace PBIK
 		// we have to be careful here when calculating the distance to the parent sub-root.
 		// if the parent sub-root is attached to an effector, use the effector's position
 		// otherwise use the current position of the FRigidBody
-		const FEffector* ParentEffector = ParentSubRoot->Effector;
-		const FVector ParentSubRootPosition = ParentEffector ? ParentEffector->Position : ParentSubRoot->Position;
+		const FEffector* ParentEffector = ChainRootBody->Effector;
+		const FVector ParentSubRootPosition = ParentEffector ? ParentEffector->Position : ChainRootBody->Position;
 		const float DistToSubRootCurrent = (ParentSubRootPosition - Position).Size();
-		if (DistToSubRootCurrent >= DistToSubRootInInputPose)
+		if (DistToSubRootCurrent >= DistToChainRootInInputPose)
 		{
 			return; // limb is stretched
 		}
 
 		// amount squashed (clamped to scaled original length)
-		const float DeltaSquash = DistToSubRootInRefPose - DistToSubRootCurrent;
-		float SquashPercent = DeltaSquash / DistToSubRootInRefPose;
+		const float DeltaSquash = DistToChainRootInInputPose - DistToSubRootCurrent;
+		float SquashPercent = DeltaSquash / DistToChainRootInInputPose;
 		SquashPercent = PBIK::CircularEaseOut(SquashPercent);
 		if (FMath::IsNearlyZero(SquashPercent))
 		{
@@ -117,10 +163,9 @@ namespace PBIK
 				const FRigidBody* CurrentParent = Body->GetParentBody();
 				FQuat LocalRotation = CurrentParent ? CurrentParent->Rotation.Inverse() * Body->Rotation : Body->Rotation;
 				const FRotator AnglesFromInput = LocalRotation.Rotator();
-				const FRotator AnglesFromRefPose = Body->Bone->LocalRotationInitial;
 				const FRotator PreferredAngles = Body->J.PreferredAngles * SquashPercent;
 
-				auto CalcDeltaAngle = [](double PreferredAngle, double RefPoseAngle, double InputAngle) -> double
+				auto CalcDeltaAngle = [](double PreferredAngle, double InputAngle) -> double
 				{
 					if (FMath::Abs(PreferredAngle) <= SMALL_NUMBER)
 					{
@@ -128,10 +173,9 @@ namespace PBIK
 					}
 
 					// if the target angle is greater than the input pose, return a delta
-					const float TargetAngle = PreferredAngle + RefPoseAngle;
-					if (FMath::Abs(TargetAngle) > FMath::Abs(InputAngle))
+					if (FMath::Abs(PreferredAngle) > FMath::Abs(InputAngle))
 					{
-						return TargetAngle - InputAngle;
+						return PreferredAngle - InputAngle;
 					}
 
 					// body is already bent more than the preferred angle
@@ -139,14 +183,14 @@ namespace PBIK
 				};
 
 				FRotator LocalDeltaAngles;
-				LocalDeltaAngles.Pitch = CalcDeltaAngle(PreferredAngles.Pitch, AnglesFromRefPose.Pitch, AnglesFromInput.Pitch);
-				LocalDeltaAngles.Yaw = CalcDeltaAngle(PreferredAngles.Yaw, AnglesFromRefPose.Yaw, AnglesFromInput.Yaw);
-				LocalDeltaAngles.Roll = CalcDeltaAngle(PreferredAngles.Roll, AnglesFromRefPose.Roll, AnglesFromInput.Roll);
+				LocalDeltaAngles.Pitch = CalcDeltaAngle(PreferredAngles.Pitch, AnglesFromInput.Pitch);
+				LocalDeltaAngles.Yaw = CalcDeltaAngle(PreferredAngles.Yaw, AnglesFromInput.Yaw);
+				LocalDeltaAngles.Roll = CalcDeltaAngle(PreferredAngles.Roll, AnglesFromInput.Roll);
 				Body->Rotation = Body->Rotation * LocalDeltaAngles.Quaternion();
 				Body->Rotation.Normalize();
 			}
 
-			if (Parent == ParentSubRoot->Bone)
+			if (Parent == ChainRootBody->Bone)
 			{
 				return;
 			}
@@ -353,12 +397,12 @@ void FPBIKSolver::ApplyPullChainAlpha(const float GlobalPullChainAlpha)
 	
 	for (FEffector& Effector : Effectors)
 	{
-		if (!Effector.ParentSubRoot)
+		if (!Effector.ChainRootBody)
 		{
 			continue;
 		}
 
-		if (Effector.DistToSubRootAlongBones < SMALL_NUMBER)
+		if (Effector.DistToChainRootAlongBones < SMALL_NUMBER)
 		{
 			continue;
 		}
@@ -369,7 +413,7 @@ void FPBIKSolver::ApplyPullChainAlpha(const float GlobalPullChainAlpha)
 		}
 
 		// get original chain vector
-		const FVector ChainStartOrig = Effector.ParentSubRoot->InputPosition;
+		const FVector ChainStartOrig = Effector.ChainRootBody->InputPosition;
 		const FVector ChainEndOrig = Effector.PositionOrig;
 		FVector ChainVecOrig;
 		float ChainLenOrig;
@@ -377,15 +421,15 @@ void FPBIKSolver::ApplyPullChainAlpha(const float GlobalPullChainAlpha)
 
 		// get new chain vector
 		FVector ChainStartNew;
-		if (Effector.ParentSubRoot->Effector)
+		if (Effector.ChainRootBody->Effector)
 		{
 			// add effector offset to chain root
-			const FEffector* RootEffector = Effector.ParentSubRoot->Effector;
+			const FEffector* RootEffector = Effector.ChainRootBody->Effector;
 			FVector RootEffectorDelta = RootEffector->Position - RootEffector->PositionOrig;
-			ChainStartNew = Effector.ParentSubRoot->Position + RootEffectorDelta;
+			ChainStartNew = Effector.ChainRootBody->Position + RootEffectorDelta;
 		}else
 		{
-			ChainStartNew = Effector.ParentSubRoot->Position;
+			ChainStartNew = Effector.ChainRootBody->Position;
 		}
 		const FVector ChainEndNew = Effector.Position;
 		FVector ChainVecNew;
@@ -399,8 +443,8 @@ void FPBIKSolver::ApplyPullChainAlpha(const float GlobalPullChainAlpha)
 
 		const FBone* Bone = Effector.Bone->Parent;
 		int32 ChainIndex = 0;
-		const float InvChainLength = 1.0f / Effector.DistToSubRootAlongBones;
-		while (Bone && Bone->Body != Effector.ParentSubRoot)
+		const float InvChainLength = 1.0f / Effector.DistToChainRootAlongBones;
+		while (Bone && Bone->Body != Effector.ChainRootBody)
 		{		
 			// rotate body along with chain
 			const FVector BodyRelativeToChain = Bone->Body->Position - ChainStartNew;
@@ -480,11 +524,6 @@ bool FPBIKSolver::Initialize()
 		return false;
 	}
 
-	if (!InitEffectors())
-    {
-    	return false;
-    }
-
 	if (!InitConstraints())
 	{
 		return false;
@@ -651,35 +690,7 @@ bool FPBIKSolver::InitBodies()
 	// must be done AFTER setting: Bone.IsSubRoot/IsSolverRoot/Parent
 	for (FEffector& Effector : Effectors)
 	{
-		const FBone* Parent = Effector.Bone->Parent;
-		while (Parent)
-		{
-			if (!Parent->bIsSolved)
-			{
-				break; // this only happens when effector is on solver root
-			}
-
-			if (Parent->bIsSubRoot || Parent->bIsSolverRoot)
-			{
-				Effector.ParentSubRoot = Parent->Body;
-				break;
-			}
-
-			Parent = Parent->Parent;
-		}
-	}
-
-	return true;
-}
-
-bool FPBIKSolver::InitEffectors()
-{
-	using PBIK::FEffector;
-
-	// store the length of chain this effector controls in the reference pose
-	for (FEffector& Effector : Effectors)
-	{
-		Effector.DistToSubRootInRefPose = Effector.CalculateDistanceToSubRoot();
+		Effector.UpdateChainRoot();
 	}
 
 	return true;
@@ -789,7 +800,7 @@ int32 FPBIKSolver::AddBone(
 	return Bones.Emplace(Name, ParentIndex, InOrigPosition, InOrigRotation, bIsSolverRoot);
 }
 
-int32 FPBIKSolver::AddEffector(FName BoneName)
+int32 FPBIKSolver::AddEffector(const FName BoneName)
 {
 	for (PBIK::FBone& Bone : Bones)
 	{
diff --git a/Engine/Plugins/Experimental/FullBodyIK/Source/PBIK/Private/RigUnit_PBIK.cpp b/Engine/Plugins/Experimental/FullBodyIK/Source/PBIK/Private/RigUnit_PBIK.cpp
--- a/Engine/Plugins/Experimental/FullBodyIK/Source/PBIK/Private/RigUnit_PBIK.cpp
+++ b/Engine/Plugins/Experimental/FullBodyIK/Source/PBIK/Private/RigUnit_PBIK.cpp
@@ -179,6 +179,7 @@ FRigUnit_PBIK_Execute()
 		EffectorSettings.PositionAlpha = Effector.PositionAlpha;
 		EffectorSettings.RotationAlpha = Effector.RotationAlpha;
 		EffectorSettings.StrengthAlpha = Effector.StrengthAlpha;
+		EffectorSettings.ChainDepth = Effector.ChainDepth;
 		EffectorSettings.PullChainAlpha = Effector.PullChainAlpha;
 		EffectorSettings.PinRotation = Effector.PinRotation;
 		
diff --git a/Engine/Plugins/Experimental/FullBodyIK/Source/PBIK/Public/Core/PBIKSolver.h b/Engine/Plugins/Experimental/FullBodyIK/Source/PBIK/Public/Core/PBIKSolver.h
--- a/Engine/Plugins/Experimental/FullBodyIK/Source/PBIK/Public/Core/PBIKSolver.h
+++ b/Engine/Plugins/Experimental/FullBodyIK/Source/PBIK/Public/Core/PBIKSolver.h
@@ -45,6 +45,13 @@ struct FEffectorSettings
 	UPROPERTY(EditAnywhere, Category="Effector")
 	float StrengthAlpha = 1.0f;
 
+	/** Range 0-inf (default is 0). Explicitly set the number of bones up the hierarchy to consider part of this effector's 'chain'.
+	* The "chain" of bones is used to apply Preferred Angles, Pull Chain Alpha and Chain "Sub Solves".
+	* If left at 0, the solver will attempt to determine the root of the chain by searching up the hierarchy until it finds a branch or another effector, whichever it finds first.
+	*/
+	UPROPERTY(BlueprintReadWrite, EditAnywhere, Category = "Full Body IK Effector", meta = (ClampMin = "0", UIMin = "0"))
+	int32 ChainDepth = 0;
+
 	/** Range 0-1 (default is 1.0). When enabled (greater than 0.0), the solver internally partitions the skeleton into 'chains' which extend from the effector to the nearest fork in the skeleton.
 	*These chains are pre-rotated and translated, as a whole, towards the effector targets.
 	*This can improve the results for sparse bone chains, and significantly improve convergence on dense bone chains.
@@ -74,12 +81,12 @@ struct FEffector
 
 	FBone* Bone;
 	TWeakPtr<FPinConstraint> Pin;
-	FRigidBody* ParentSubRoot = nullptr;
-	float DistToSubRootInInputPose;
-	float DistToSubRootInRefPose;
+	FRigidBody* ChainRootBody = nullptr;
+	int32 ChainRootDepthInitializedWith = -1;
+	float DistToChainRootInInputPose;
 	
 	TArray<float> DistancesFromEffector;
-	float DistToSubRootAlongBones;
+	float DistToChainRootAlongBones;
 
 	FEffector(FBone* InBone);
 
@@ -88,8 +95,9 @@ struct FEffector
 		const FQuat& InRotationGoal,
 		const FEffectorSettings& InSettings);
 
+	void UpdateChainRoot();
 	void UpdateFromInputs(const FBone& SolverRoot);
-	float CalculateDistanceToSubRoot() const;
+	float CalculateDistanceToChainRoot() const;
 	void ApplyPreferredAngles() const;
 };
 	
@@ -240,7 +248,7 @@ public:
 		const FQuat& InOrigRotation,
 		bool bIsSolverRoot);
 
-	int32 AddEffector(FName BoneName);
+	int32 AddEffector(const FName BoneName);
 	
 private:
 
@@ -248,8 +256,6 @@ private:
 
 	bool InitBodies();
 
-	bool InitEffectors();
-
 	bool InitConstraints();
 
 	void AddBodyForBone(PBIK::FBone* Bone);
diff --git a/Engine/Plugins/Experimental/FullBodyIK/Source/PBIK/Public/RigUnit_PBIK.h b/Engine/Plugins/Experimental/FullBodyIK/Source/PBIK/Public/RigUnit_PBIK.h
--- a/Engine/Plugins/Experimental/FullBodyIK/Source/PBIK/Public/RigUnit_PBIK.h
+++ b/Engine/Plugins/Experimental/FullBodyIK/Source/PBIK/Public/RigUnit_PBIK.h
@@ -60,6 +60,7 @@ struct FPBIKEffector
 		PositionAlpha(1.0f),
 		RotationAlpha(1.0f),
 		StrengthAlpha(1.0f),
+		ChainDepth(0),
 		PullChainAlpha(1.0f),
 		PinRotation(1.0f){}
 
@@ -86,6 +87,13 @@ struct FPBIKEffector
 	UPROPERTY(EditAnywhere, Category="Effector", meta = (ClampMin = "0.0", ClampMax = "1.0", UIMin = "0.0", UIMax = "1.0"))
 	float StrengthAlpha = 1.0f;
 
+	/** Range 0-inf (default is 0). Explicitly set the number of bones up the hierarchy to consider part of this effector's 'chain'.
+	* The "chain" of bones is used to apply Preferred Angles, Pull Chain Alpha and Chain "Sub Solves".
+	* If left at 0, the solver will attempt to determine the root of the chain by searching up the hierarchy until it finds a branch or another effector, whichever it finds first.
+	*/
+	UPROPERTY(BlueprintReadWrite, EditAnywhere, Category = "Full Body IK Effector", meta = (ClampMin = "0", UIMin = "0"))
+	int32 ChainDepth = 0;
+
 	/** Range 0-1 (default is 1.0). When enabled (greater than 0.0), the solver internally partitions the skeleton into 'chains' which extend from the effector to the nearest fork in the skeleton.
 	 *These chains are pre-rotated and translated, as a whole, towards the effector targets.
 	 *This can improve the results for sparse bone chains, and significantly improve convergence on dense bone chains.
