Control Rig: Fix FullBodyIK not working with control rig debugger
#rnx

diff --git a/Engine/Plugins/Experimental/FullBodyIK/Source/PBIK/Public/Core/PBIKSolver.h b/Engine/Plugins/Experimental/FullBodyIK/Source/PBIK/Public/Core/PBIKSolver.h
--- a/Engine/Plugins/Experimental/FullBodyIK/Source/PBIK/Public/Core/PBIKSolver.h
+++ b/Engine/Plugins/Experimental/FullBodyIK/Source/PBIK/Public/Core/PBIKSolver.h
@@ -203,6 +203,12 @@ struct PBIK_API FPBIKSolver
 
 public:
 
+	FPBIKSolver& operator=(const FPBIKSolver&)
+	{
+		bReadyToSimulate = false;
+		return *this;
+	}
+
 	PBIK::FDebugDraw* GetDebugDraw();
 
 	//
