#UE Emit a warning about failing to initialize a FBIK solver including the path to an object that helps identify what asset failed.

diff --git a/Engine/Plugins/Experimental/FullBodyIK/Source/PBIK/Private/RigUnit_PBIK.cpp b/Engine/Plugins/Experimental/FullBodyIK/Source/PBIK/Private/RigUnit_PBIK.cpp
--- a/Engine/Plugins/Experimental/FullBodyIK/Source/PBIK/Private/RigUnit_PBIK.cpp
+++ b/Engine/Plugins/Experimental/FullBodyIK/Source/PBIK/Private/RigUnit_PBIK.cpp
@@ -6,6 +6,7 @@
 #include "Rigs/RigHierarchy.h"
 #include "Rigs/RigHierarchyElements.h"
 #include "Misc/HashBuilder.h"
+#include "PBIK.h"
 
 #include UE_INLINE_GENERATED_CPP_BY_NAME(RigUnit_PBIK)
 
@@ -110,7 +111,10 @@ FRigUnit_PBIK_Execute()
 		}
 		
 		// initialize
-		WorkData.Solver.Initialize();
+		if (!WorkData.Solver.Initialize())
+		{
+			UE_LOG(LogPBIKSolver, Warning, TEXT("PBIK: Solver failed to initialize in '%s'"), *GetPathNameSafe(Hierarchy));
+		}
 		WorkData.bNeedsInit = false;
 		WorkData.HashInitializedWith = Hash;
 	}
