Fixes Pin Rotation in Full Body IK solver.
Fixes bug, intermediate effector pin rotation not working.

diff --git a/Engine/Plugins/Experimental/FullBodyIK/Source/PBIK/Private/Core/PBIKSolver.cpp b/Engine/Plugins/Experimental/FullBodyIK/Source/PBIK/Private/Core/PBIKSolver.cpp
--- a/Engine/Plugins/Experimental/FullBodyIK/Source/PBIK/Private/Core/PBIKSolver.cpp
+++ b/Engine/Plugins/Experimental/FullBodyIK/Source/PBIK/Private/Core/PBIKSolver.cpp
@@ -296,21 +296,18 @@ void FPBIKSolver::UpdateBonesFromBodies()
 			continue; // if there's an effector on the root, leave it where the body ended up
 		}
 
-		if (Bone->Body)
+		// if effector is between other effectors, it will have an associated rigid body
+		// so leave transform where body ended up after solve, otherwise snap position back to FK location
+		if (!Bone->Body)
 		{
-			continue; // effector is between other effectors, so leave transform where body ended up
+			Bone->Position = Bone->Parent->Position + Bone->Parent->Rotation * Bone->LocalPositionFromInput;
 		}
-		
-		Bone->Position = Bone->Parent->Position + Bone->Parent->Rotation * Bone->LocalPositionFromInput;
 
+		// optionally pin rotation to that of effector
 		if (Effector.Settings.PinRotation > SMALL_NUMBER)
 		{
-			// optionally pin rotation to that of effector
-			const float RotAmount = FMath::Clamp(0.0f, 1.0f, Effector.Settings.PinRotation);
+			const float RotAmount = FMath::Clamp(Effector.Settings.PinRotation, 0.0f, 1.0f);
 			Bone->Rotation = FQuat::FastLerp(Bone->Rotation, Effector.Rotation, RotAmount).GetNormalized();
-		}else
-		{
-			Bone->Rotation = Bone->Parent->Rotation * Bone->LocalRotationFromInput;
 		}
 	}
 
