Added sub-solve capability to core FBIK solver

diff --git a/Engine/Plugins/Animation/IKRig/Source/IKRig/Private/Rig/Solvers/IKRig_PBIKSolver.cpp b/Engine/Plugins/Animation/IKRig/Source/IKRig/Private/Rig/Solvers/IKRig_PBIKSolver.cpp
--- a/Engine/Plugins/Animation/IKRig/Source/IKRig/Private/Rig/Solvers/IKRig_PBIKSolver.cpp
+++ b/Engine/Plugins/Animation/IKRig/Source/IKRig/Private/Rig/Solvers/IKRig_PBIKSolver.cpp
@@ -118,6 +118,7 @@ void UIKRigFBIKSolver::Solve(FIKRigSkeleton& IKRigSkeleton, const FIKRigGoalCont
 	// update settings
 	FPBIKSolverSettings Settings;
 	Settings.Iterations = Iterations;
+	Settings.SubIterations = SubIterations;
 	Settings.MassMultiplier = MassMultiplier;
 	Settings.bAllowStretch = bAllowStretch;
 	Settings.RootBehavior = RootBehavior;
@@ -153,6 +154,7 @@ void UIKRigFBIKSolver::UpdateSolverSettings(UIKRigSolver* InSettings)
 	if(UIKRigFBIKSolver* Settings = Cast<UIKRigFBIKSolver>(InSettings))
 	{
 		Iterations = Settings->Iterations;
+		SubIterations = Settings->SubIterations;
 		MassMultiplier = Settings->MassMultiplier;
 		bAllowStretch = Settings->bAllowStretch;
 		RootBehavior = Settings->RootBehavior;
diff --git a/Engine/Plugins/Animation/IKRig/Source/IKRig/Public/Rig/Solvers/IKRig_FBIKSolver.h b/Engine/Plugins/Animation/IKRig/Source/IKRig/Public/Rig/Solvers/IKRig_FBIKSolver.h
--- a/Engine/Plugins/Animation/IKRig/Source/IKRig/Public/Rig/Solvers/IKRig_FBIKSolver.h
+++ b/Engine/Plugins/Animation/IKRig/Source/IKRig/Public/Rig/Solvers/IKRig_FBIKSolver.h
@@ -200,6 +200,10 @@ public:
 	UPROPERTY(BlueprintReadWrite, EditAnywhere, Category = SolverSettings, meta = (ClampMin = "0", ClampMax = "1000", UIMin = "0.0", UIMax = "200.0"))
 	int32 Iterations = 20;
 
+	/** Iterations used for sub-chains defined by the Chain Depth of the effectors. These are solved BEFORE the main iteration pass. Default is 0. */
+	UPROPERTY(BlueprintReadWrite, EditAnywhere, Category = SolverSettings, meta = (ClampMin = "0", ClampMax = "1000", UIMin = "0.0", UIMax = "200.0"))
+	int32 SubIterations = 0;
+
 	/** A global mass multiplier; higher values will make the joints more stiff, but require more iterations. Typical range is 0.0 to 10.0. */
 	UPROPERTY(BlueprintReadWrite, EditAnywhere, Category = SolverSettings, meta = (ClampMin = "0", UIMin = "0.0", UIMax = "10.0"))
 	float MassMultiplier = 1.0f;
diff --git a/Engine/Plugins/Experimental/FullBodyIK/Source/PBIK/Private/Core/PBIKBody.cpp b/Engine/Plugins/Experimental/FullBodyIK/Source/PBIK/Private/Core/PBIKBody.cpp
--- a/Engine/Plugins/Experimental/FullBodyIK/Source/PBIK/Private/Core/PBIKBody.cpp
+++ b/Engine/Plugins/Experimental/FullBodyIK/Source/PBIK/Private/Core/PBIKBody.cpp
@@ -119,21 +119,22 @@ FRigidBody* FRigidBody::GetParentBody() const
 	return nullptr;
 }
 
-float FRigidBody::GetInverseMass()
+bool FRigidBody::IsAllowedToRotate() const
 {
-	if (Pin && Pin->bEnabled)
+	const bool bIsLockedByPin= Pin && Pin->bEnabled && Pin->bLockRotation;
+	if (bIsLockedByPin || bIsLockedBySubSolve)
 	{
-		return 1.0f - Pin->Alpha;
+		return false; // body not allowed to rotate
 	}
 
-	return InvMass;
+	return true;
 }
 
 void FRigidBody::ApplyPushToRotateBody(const FVector& Push, const FVector& Offset)
 {
-	if (Pin && Pin->bEnabled && Pin->bPinRotation)
+	if (!IsAllowedToRotate())
 	{
-		return; // rotation of this body is pinned
+		return; // rotation of this body is disabled
 	}
 	
 	// equation 8 in "Detailed Rigid Body Simulation with XPBD"
@@ -149,9 +150,9 @@ void FRigidBody::ApplyPushToPosition(const FVector& Push)
 
 void FRigidBody::ApplyRotationDelta(const FQuat& DeltaQ)
 {
-	if (Pin && Pin->bEnabled && Pin->bPinRotation)
+	if (!IsAllowedToRotate())
 	{
-		return; // rotation of this body is pinned
+		return; // rotation of this body is disabled
 	}
 
 	// limit rotation each iteration
diff --git a/Engine/Plugins/Experimental/FullBodyIK/Source/PBIK/Private/Core/PBIKConstraint.cpp b/Engine/Plugins/Experimental/FullBodyIK/Source/PBIK/Private/Core/PBIKConstraint.cpp
--- a/Engine/Plugins/Experimental/FullBodyIK/Source/PBIK/Private/Core/PBIKConstraint.cpp
+++ b/Engine/Plugins/Experimental/FullBodyIK/Source/PBIK/Private/Core/PBIKConstraint.cpp
@@ -25,6 +25,12 @@ FJointConstraint::FJointConstraint(FRigidBody* InA, FRigidBody* InB)
 	
 void FJointConstraint::Solve(const FPBIKSolverSettings& Settings)
 {
+	// early out if both bodies are locked (joint cannot do anything)
+	if (A->bIsLockedBySubSolve && B->bIsLockedBySubSolve)
+	{
+		return;
+	}
+	
 	// get pos correction to use for rotation
 	FVector OffsetA;
 	FVector OffsetB;
@@ -38,9 +44,9 @@ void FJointConstraint::Solve(const FPBIKSolverSettings& Settings)
 	// enforce joint limits
 	UpdateJointLimits();
 
-	// calc inv mass of body A
-	const float WA = A->GetInverseMass();
-	const float WB = B->GetInverseMass();
+	// calc inv mass of body A and B
+	const float WA = A->InvMass;
+	const float WB = B->InvMass;
 	const float W = WA + WB;
 	if (FMath::IsNearlyZero(W))
 	{
@@ -272,7 +278,7 @@ FPinConstraint::FPinConstraint(
 	FRigidBody* InBody,
 	const FVector& InPinPositionOrig,
 	const FQuat& InPinRotationOrig,
-	const bool bInPinRotation)
+	const bool bInLockRotation)
 {
 	GoalPosition = InPinPositionOrig;
 	GoalRotation = InPinRotationOrig;
@@ -281,7 +287,7 @@ FPinConstraint::FPinConstraint(
 	PinPointLocalToA = A->Rotation.Inverse() * (GoalPosition - A->Position);
 
 	ARotLocalToPin = A->Rotation * GoalRotation.Inverse();
-	bPinRotation = bInPinRotation;
+	bLockRotation = bInLockRotation;
 }
 
 void FPinConstraint::Solve(const FPBIKSolverSettings& Settings)
@@ -291,7 +297,7 @@ void FPinConstraint::Solve(const FPBIKSolverSettings& Settings)
 		return;
 	}
 
-	if (bPinRotation)
+	if (bLockRotation)
 	{
 		// keep body at fixed rotation relative to the pin goal rotation
 		A->Rotation = ARotLocalToPin * GoalRotation;
diff --git a/Engine/Plugins/Experimental/FullBodyIK/Source/PBIK/Private/Core/PBIKSolver.cpp b/Engine/Plugins/Experimental/FullBodyIK/Source/PBIK/Private/Core/PBIKSolver.cpp
--- a/Engine/Plugins/Experimental/FullBodyIK/Source/PBIK/Private/Core/PBIKSolver.cpp
+++ b/Engine/Plugins/Experimental/FullBodyIK/Source/PBIK/Private/Core/PBIKSolver.cpp
@@ -31,7 +31,7 @@ namespace PBIK
 		Settings = InSettings;
 	}
 
-	void FEffector::UpdateChainRoot()
+	void FEffector::UpdateChainStates()
 	{
 		// NOTE: this function must be called AFTER InitBodies() due to it's reliance on the FBody.bIsSubRoot flag
 		
@@ -46,30 +46,22 @@ namespace PBIK
 				break; // this only happens when effector is on solver root
 			}
 
-			if (MaxDepth > 0)
+			// if user specified a custom chain depth, then use it.
+			// otherwise halt when we hit a sub root (a branch) or root of the solver
+			const bool bIsAtAnyRoot = Parent->bIsSubRoot || Parent->bIsSolverRoot;
+			const bool bIsAtChainRoot = MaxDepth > 0 ? Depth > MaxDepth || bIsAtAnyRoot : bIsAtAnyRoot;
+			if (bIsAtChainRoot)
 			{
-				// user specified a custom chain depth, so use it
-				if (Depth >= MaxDepth)
-				{
-					ChainRootBody = Parent->Body;
-					break;
-				}
-			}
-			else
-			{
-				// user did not specify a max depth, so use the nearest sub-root as the chain root
-				if (Parent->bIsSubRoot || Parent->bIsSolverRoot)
-				{
-					ChainRootBody = Parent->Body;
-					break;
-				}
+				ChainRootBody = Parent->Body;
+				break;
 			}
-
+			
+			Parent->Body->bIsPartOfSubChain = true;
 			++Depth;
 			Parent = Parent->Parent;
 		}
 
-		ChainRootDepthInitializedWith = Settings.ChainDepth;
+		ChainDepthInitializedWith = Settings.ChainDepth;
 	}
 
 	void FEffector::UpdateFromInputs(const FBone& SolverRoot)
@@ -78,12 +70,7 @@ namespace PBIK
 		Position = FMath::Lerp(PositionOrig, PositionGoal, Settings.PositionAlpha);
 		Rotation = FMath::Lerp(RotationOrig, RotationGoal, Settings.RotationAlpha);
 		Pin.Pin()->SetGoal(Position, Rotation, Settings.StrengthAlpha);
-
-		// in case user modified the chain depth at runtime, we may need to find the new chain root 
-		if (Settings.ChainDepth != ChainRootDepthInitializedWith)
-		{
-			UpdateChainRoot();
-		}
+		
 		// update length of chain this effector controls in the input pose
 		DistToChainRootInInputPose = CalculateDistanceToChainRoot();
 
@@ -243,6 +230,9 @@ void FPBIKSolver::Solve(const FPBIKSolverSettings& Settings)
 		RootPin.Pin()->SetGoal(SolverRoot->Position, SolverRoot->Rotation, 1.0f);
 	}
 
+	// lazily updates effector chain depths (IFF settings change at runtime)
+	UpdateEffectorDepths();
+
 	// blend effectors by Alpha, update pin goals and update effector dist to root
 	for (FEffector& Effector : Effectors)
 	{
@@ -480,7 +470,7 @@ void FPBIKSolver::SolveConstraints(const FPBIKSolverSettings& Settings)
 {
 	using PBIK::FRigidBody;
 
-	for (int32  I = 0; I < Settings.Iterations; ++I)
+	auto RunConstraintPass = [this, &Settings](int32 CurrentIteration, int32 MaxIterations)
 	{
 		// solve all constraints
 		for (const TSharedPtr<PBIK::FConstraint>& Constraint : Constraints)
@@ -491,12 +481,41 @@ void FPBIKSolver::SolveConstraints(const FPBIKSolverSettings& Settings)
 		// do post-pass to remove stretch
 		if (!Settings.bAllowStretch)
 		{
-			const float StretchAmount = FMath::Min(1.0f, static_cast<float>(I+1) / (Settings.Iterations / 2.0f));
+			const float StretchAmount = FMath::Min(1.0f, static_cast<float>(CurrentIteration+1) / (MaxIterations / 2.0f));
 			for (int32 C = Constraints.Num() - 1; C >= 0; --C)
 			{
 				Constraints[C]->RemoveStretch(StretchAmount);
 			}
 		}
+		
+	};
+
+	// run pre-pass where we lock everything not in a sub-chain
+	if (bHasSubChains && Settings.SubIterations > 0)
+	{
+		// first lock all the bodies not in a sub chain
+		for (FRigidBody& Body : Bodies)
+		{
+			Body.bIsLockedBySubSolve = !Body.bIsPartOfSubChain;
+		}
+		
+		// solve all constraints (locked constraints are skipped in this pass)
+		for (int32  I = 0; I < Settings.SubIterations; ++I)
+		{
+			RunConstraintPass(I, Settings.SubIterations);
+		}
+		
+		// unlock all bodies so we can proceed with normal constraint solve
+		for (FRigidBody& Body : Bodies)
+		{
+			Body.bIsLockedBySubSolve = false;
+		}
+	}
+
+	// run main constraint pass
+	for (int32  I = 0; I < Settings.Iterations; ++I)
+	{
+		RunConstraintPass(I, Settings.Iterations);
 	}
 }
 
@@ -595,16 +614,17 @@ bool FPBIKSolver::InitBones()
 	// walk upwards from each effector to the root to initialize "Bone.IsSolved"
 	for (const FEffector& Effector : Effectors)
 	{
-		FBone* NextBone = Effector.Bone;
-		while (NextBone)
+		FBone* CurrentBone = Effector.Bone;
+		CurrentBone->bIsSolved = true;
+		while (CurrentBone)
 		{
-			NextBone->bIsSolved = true;
-			NextBone = NextBone->Parent;
-			if (NextBone && NextBone->bIsSolverRoot)
+			if (CurrentBone->bIsSolverRoot)
 			{
-				NextBone->bIsSolved = true;
+				CurrentBone->bIsSolved = true;
 				break;
 			}
+			CurrentBone->bIsSolved = true;
+			CurrentBone = CurrentBone->Parent;
 		}
 	}
 
@@ -683,13 +703,6 @@ bool FPBIKSolver::InitBodies()
 		Body.Bone->Body = &Body;
 	}
 
-	// initialize Effector's nearest ParentSubRoot (FBody) pointer
-	// must be done AFTER setting: Bone.IsSubRoot/IsSolverRoot/Parent
-	for (FEffector& Effector : Effectors)
-	{
-		Effector.UpdateChainRoot();
-	}
-
 	return true;
 }
 
@@ -759,10 +772,51 @@ bool FPBIKSolver::InitConstraints()
 		RootPin = RootConstraint;
 		SolverRoot->Body->Pin = RootConstraint.Get();
 	}
+
+	// now we can set the initial effector depths
+	UpdateEffectorDepths();
 	
 	return true;
 }
 
+void FPBIKSolver::UpdateEffectorDepths()
+{
+	using PBIK::FEffector;
+	using PBIK::FBone;
+	using PBIK::FRigidBody;
+	
+	// initialize Effector's nearest ParentSubRoot (FBody) pointer
+	// must be done AFTER setting: Bone.IsSubRoot/IsSolverRoot/Parent
+	bool bIsEffectorDepthDirty = false;
+	bHasSubChains = false;
+	for (const FEffector& Effector : Effectors)
+	{
+		if (Effector.ChainDepthInitializedWith != Effector.Settings.ChainDepth)
+		{
+			bIsEffectorDepthDirty = true;
+		}
+		if (Effector.Settings.ChainDepth > 0)
+		{
+			bHasSubChains = true;
+		}
+	}
+
+	if (bIsEffectorDepthDirty)
+	{
+		// reset body states related to sub chains
+		for (FRigidBody& Body : Bodies)
+		{
+			Body.bIsPartOfSubChain = false;
+			Body.bIsLockedBySubSolve = false;
+		}
+		// walk up each effector and set chain states
+		for (FEffector& EffectorToUpdate : Effectors)
+		{
+			EffectorToUpdate.UpdateChainStates();
+		}
+	}
+}
+
 PBIK::FDebugDraw* FPBIKSolver::GetDebugDraw()
 {
 	DebugDraw.Solver = this;
diff --git a/Engine/Plugins/Experimental/FullBodyIK/Source/PBIK/Public/Core/PBIKBody.h b/Engine/Plugins/Experimental/FullBodyIK/Source/PBIK/Public/Core/PBIKBody.h
--- a/Engine/Plugins/Experimental/FullBodyIK/Source/PBIK/Public/Core/PBIKBody.h
+++ b/Engine/Plugins/Experimental/FullBodyIK/Source/PBIK/Public/Core/PBIKBody.h
@@ -92,6 +92,9 @@ struct FRigidBody
 	float Mass = 0.f;
 
 	float MaxAngle = 0.0f;
+
+	bool bIsLockedBySubSolve = false;
+	bool bIsPartOfSubChain = false;
 	
 private:
 
@@ -111,7 +114,7 @@ public:
 
 	FRigidBody* GetParentBody() const;
 
-	float GetInverseMass();
+	bool IsAllowedToRotate() const;
 
 	void ApplyPushToRotateBody(const FVector& Push, const FVector& Offset);
 	
diff --git a/Engine/Plugins/Experimental/FullBodyIK/Source/PBIK/Public/Core/PBIKConstraint.h b/Engine/Plugins/Experimental/FullBodyIK/Source/PBIK/Public/Core/PBIKConstraint.h
--- a/Engine/Plugins/Experimental/FullBodyIK/Source/PBIK/Public/Core/PBIKConstraint.h
+++ b/Engine/Plugins/Experimental/FullBodyIK/Source/PBIK/Public/Core/PBIKConstraint.h
@@ -100,7 +100,7 @@ private:
 	FVector PinPointLocalToA;
 	float Alpha = 1.0;
 	FQuat ARotLocalToPin;
-	bool bPinRotation = false;
+	bool bLockRotation = false;
 
 public:
 	
@@ -108,7 +108,7 @@ public:
 		FRigidBody* InBody,
 		const FVector& InPinPositionOrig,
 		const FQuat& InPinRotationOrig,
-		const bool bInPinRotation);
+		const bool bInLockRotation);
 
 	virtual ~FPinConstraint() {};
 
diff --git a/Engine/Plugins/Experimental/FullBodyIK/Source/PBIK/Public/Core/PBIKSolver.h b/Engine/Plugins/Experimental/FullBodyIK/Source/PBIK/Public/Core/PBIKSolver.h
--- a/Engine/Plugins/Experimental/FullBodyIK/Source/PBIK/Public/Core/PBIKSolver.h
+++ b/Engine/Plugins/Experimental/FullBodyIK/Source/PBIK/Public/Core/PBIKSolver.h
@@ -82,7 +82,7 @@ struct FEffector
 	FBone* Bone;
 	TWeakPtr<FPinConstraint> Pin;
 	FRigidBody* ChainRootBody = nullptr;
-	int32 ChainRootDepthInitializedWith = -1;
+	int32 ChainDepthInitializedWith = -1;
 	float DistToChainRootInInputPose;
 	
 	TArray<float> DistancesFromEffector;
@@ -95,7 +95,7 @@ struct FEffector
 		const FQuat& InRotationGoal,
 		const FEffectorSettings& InSettings);
 
-	void UpdateChainRoot();
+	void UpdateChainStates();
 	void UpdateFromInputs(const FBone& SolverRoot);
 	float CalculateDistanceToChainRoot() const;
 	void ApplyPreferredAngles() const;
@@ -158,6 +158,10 @@ struct PBIK_API FPBIKSolverSettings
 	UPROPERTY(BlueprintReadWrite, EditAnywhere, Category = SolverSettings, meta = (ClampMin = "0", ClampMax = "1000", UIMin = "0.0", UIMax = "200.0"))
 	int32 Iterations = 20;
 
+	/** Iterations used for sub-chains defined by the Chain Depth of the effectors. These are solved BEFORE the main iteration pass. Default is 0. */
+	UPROPERTY(BlueprintReadWrite, EditAnywhere, Category = SolverSettings, meta = (ClampMin = "0", ClampMax = "1000", UIMin = "0.0", UIMax = "200.0"))
+	int32 SubIterations = 0;
+
 	/** A global mass multiplier; higher values will make the joints more stiff, but require more iterations. Typical range is 0.0 to 10.0. */
 	UPROPERTY(BlueprintReadWrite, EditAnywhere, Category = SolverSettings, meta = (ClampMin = "0", UIMin = "0.0", UIMax = "10.0"))
 	float MassMultiplier = 1.0f;
@@ -264,6 +268,8 @@ private:
 
 	bool InitConstraints();
 
+	void UpdateEffectorDepths();
+
 	void AddBodyForBone(PBIK::FBone* Bone);
 
 	void UpdateBodies(const FPBIKSolverSettings& Settings);
@@ -295,6 +301,7 @@ private:
 	TArray<PBIK::FRigidBody> Bodies;
 	TArray<TSharedPtr<PBIK::FConstraint>> Constraints;
 	TArray<PBIK::FEffector> Effectors;
+	bool bHasSubChains = false;
 	bool bReadyToSimulate = false;
 	
 	PBIK::FDebugDraw DebugDraw;
diff --git a/Engine/Plugins/Experimental/FullBodyIK/Source/PBIK/Public/RigUnit_PBIK.h b/Engine/Plugins/Experimental/FullBodyIK/Source/PBIK/Public/RigUnit_PBIK.h
--- a/Engine/Plugins/Experimental/FullBodyIK/Source/PBIK/Public/RigUnit_PBIK.h
+++ b/Engine/Plugins/Experimental/FullBodyIK/Source/PBIK/Public/RigUnit_PBIK.h
@@ -88,7 +88,7 @@ struct FPBIKEffector
 	float StrengthAlpha = 1.0f;
 
 	/** Range 0-inf (default is 0). Explicitly set the number of bones up the hierarchy to consider part of this effector's 'chain'.
-	* The "chain" of bones is used to apply Preferred Angles, Pull Chain Alpha and Chain "Sub Solves".
+	* The "chain" of bones is used to apply Preferred Angles, Pull Chain Alpha and "Sub Solve" iterations.
 	* If left at 0, the solver will attempt to determine the root of the chain by searching up the hierarchy until it finds a branch or another effector, whichever it finds first.
 	*/
 	UPROPERTY(BlueprintReadWrite, EditAnywhere, Category = "Full Body IK Effector", meta = (ClampMin = "0", UIMin = "0"))
