Core FBIK solver fixes and simplification.
Fixes issue, tip bones in full body IK solver do not inherit parent rotation.
Fixes issue, "Pin to Input" not working root of skeleton.

diff --git a/Engine/Plugins/Experimental/FullBodyIK/Source/PBIK/Private/Core/PBIKBody.cpp b/Engine/Plugins/Experimental/FullBodyIK/Source/PBIK/Private/Core/PBIKBody.cpp
--- a/Engine/Plugins/Experimental/FullBodyIK/Source/PBIK/Private/Core/PBIKBody.cpp
+++ b/Engine/Plugins/Experimental/FullBodyIK/Source/PBIK/Private/Core/PBIKBody.cpp
@@ -121,13 +121,7 @@ FRigidBody* FRigidBody::GetParentBody() const
 
 bool FRigidBody::IsAllowedToRotate() const
 {
-	const bool bIsLockedByPin= Pin && Pin->bEnabled && Pin->bLockRotation;
-	if (bIsLockedByPin || bIsLockedBySubSolve)
-	{
-		return false; // body not allowed to rotate
-	}
-
-	return true;
+	return !bIsLockedBySubSolve && InvMass > SMALL_NUMBER;
 }
 
 void FRigidBody::ApplyPushToRotateBody(const FVector& Push, const FVector& Offset)
@@ -143,9 +137,9 @@ void FRigidBody::ApplyPushToRotateBody(const FVector& Push, const FVector& Offse
 	ApplyRotationDelta(DeltaQ);
 }
 
-void FRigidBody::ApplyPushToPosition(const FVector& Push)
+void FRigidBody::ApplyPositionDelta(const FVector& DeltaP)
 {
-	Position += Push * (1.0f - J.PositionStiffness) * SolverSettings->OverRelaxation;
+	Position += DeltaP * (1.0f - J.PositionStiffness) * SolverSettings->OverRelaxation;
 }
 
 void FRigidBody::ApplyRotationDelta(const FQuat& DeltaQ)
diff --git a/Engine/Plugins/Experimental/FullBodyIK/Source/PBIK/Private/Core/PBIKConstraint.cpp b/Engine/Plugins/Experimental/FullBodyIK/Source/PBIK/Private/Core/PBIKConstraint.cpp
--- a/Engine/Plugins/Experimental/FullBodyIK/Source/PBIK/Private/Core/PBIKConstraint.cpp
+++ b/Engine/Plugins/Experimental/FullBodyIK/Source/PBIK/Private/Core/PBIKConstraint.cpp
@@ -65,8 +65,8 @@ void FJointConstraint::Solve(const FPBIKSolverSettings& Settings)
 	const FVector CorrectA = N * (C * (WA / W));
 	const FVector CorrectB = N * (C * (WB / W));
 
-	A->ApplyPushToPosition(CorrectA);
-	B->ApplyPushToPosition(-CorrectB);
+	A->ApplyPositionDelta(CorrectA);
+	B->ApplyPositionDelta(-CorrectB);
 }
 
 void FJointConstraint::RemoveStretch(const float Percent)
@@ -277,17 +277,14 @@ float FJointConstraint::SignedAngleBetweenNormals(
 FPinConstraint::FPinConstraint(
 	FRigidBody* InBody,
 	const FVector& InPinPositionOrig,
-	const FQuat& InPinRotationOrig,
-	const bool bInLockRotation)
+	const FQuat& InPinRotationOrig)
 {
 	GoalPosition = InPinPositionOrig;
 	GoalRotation = InPinRotationOrig;
 	
 	A = InBody;
 	PinPointLocalToA = A->Rotation.Inverse() * (GoalPosition - A->Position);
-
 	ARotLocalToPin = A->Rotation * GoalRotation.Inverse();
-	bLockRotation = bInLockRotation;
 }
 
 void FPinConstraint::Solve(const FPBIKSolverSettings& Settings)
@@ -297,31 +294,18 @@ void FPinConstraint::Solve(const FPBIKSolverSettings& Settings)
 		return;
 	}
 
-	if (bLockRotation)
-	{
-		// keep body at fixed rotation relative to the pin goal rotation
-		A->Rotation = ARotLocalToPin * GoalRotation;
-		
-		// move body to pin location
-		FVector AToPinPoint;
-		const FVector Correction = GetPositionCorrection(AToPinPoint);
-		A->Position += Correction;
-		
-	}else
-	{
-		// get positional correction for rotation
-		FVector AToPinPoint;
-		FVector Correction = GetPositionCorrection(AToPinPoint);
+	// get positional correction for rotation
+	FVector AToPinPoint;
+	FVector Correction = GetPositionCorrection(AToPinPoint);
 		
-		// rotate body from alignment of pin points
-		A->ApplyPushToRotateBody(Correction, AToPinPoint);
-
-		// apply positional correction to Body to align with target (after rotation)
-		// (applying directly without considering PositionStiffness because PinConstraints need
-		// to precisely pull the attached body to achieve convergence)
-		Correction = GetPositionCorrection(AToPinPoint);
-		A->Position += Correction * Settings.OverRelaxation; 
-	}
+	// rotate body from alignment of pin points
+	A->ApplyPushToRotateBody(Correction, AToPinPoint);
+
+	// apply positional correction to Body to align with target (after rotation)
+	// (applying directly without considering PositionStiffness because PinConstraints need
+	// to precisely pull the attached body to achieve convergence)
+	Correction = GetPositionCorrection(AToPinPoint);
+	A->Position += Correction * Settings.OverRelaxation; 
 }
 
 void FPinConstraint::SetGoal(const FVector& InGoalPosition, const FQuat& InGoalRotation, const float InAlpha)
@@ -331,13 +315,6 @@ void FPinConstraint::SetGoal(const FVector& InGoalPosition, const FQuat& InGoalR
 	Alpha = InAlpha;
 }
 
-void FPinConstraint::EnableInCurrentState()
-{
-	GoalPosition = A->Position;
-	GoalRotation = A->Rotation;
-	bEnabled = true;
-}
-
 FVector FPinConstraint::GetPositionCorrection(FVector& OutBodyToPinPoint) const
 {
 	OutBodyToPinPoint = A->Rotation * PinPointLocalToA;
diff --git a/Engine/Plugins/Experimental/FullBodyIK/Source/PBIK/Private/Core/PBIKSolver.cpp b/Engine/Plugins/Experimental/FullBodyIK/Source/PBIK/Private/Core/PBIKSolver.cpp
--- a/Engine/Plugins/Experimental/FullBodyIK/Source/PBIK/Private/Core/PBIKSolver.cpp
+++ b/Engine/Plugins/Experimental/FullBodyIK/Source/PBIK/Private/Core/PBIKSolver.cpp
@@ -223,11 +223,11 @@ void FPBIKSolver::Solve(const FPBIKSolverSettings& Settings)
 	}
 
 	// optionally pin root in-place (convenience, does not require an effector)
-	if (RootPin.IsValid())
+	if (Settings.RootBehavior == EPBIKRootBehavior::PinToInput)
 	{
-		RootPin.Pin()->bEnabled = Settings.RootBehavior == EPBIKRootBehavior::PinToInput;
-		// pin to animated input root pose
-		RootPin.Pin()->SetGoal(SolverRoot->Position, SolverRoot->Rotation, 1.0f);
+		// we can safely override the mass here to pin the root body
+		// if the user switches to a different root mode, the call to Body.UpdateFromInputs() will restore the correct InvMass for the body
+		SolverRoot->Body->InvMass = 0.0f;
 	}
 
 	// lazily updates effector chain depths (IFF settings change at runtime)
@@ -294,11 +294,9 @@ void FPBIKSolver::UpdateBonesFromBodies()
 		}
 
 		// optionally pin rotation to that of effector
-		if (Effector.Settings.PinRotation > SMALL_NUMBER)
-		{
-			const float RotAmount = FMath::Clamp(Effector.Settings.PinRotation, 0.0f, 1.0f);
-			Bone->Rotation = FQuat::FastLerp(Bone->Rotation, Effector.Rotation, RotAmount).GetNormalized();
-		}
+		const FQuat RotationWithoutEffector = Bone->Parent ? Bone->Parent->Rotation * Bone->LocalRotationFromInput : Bone->Rotation; 
+		const float RotAmount = FMath::Clamp(Effector.Settings.PinRotation, 0.0f, 1.0f);
+		Bone->Rotation = FQuat::FastLerp(RotationWithoutEffector, Effector.Rotation, RotAmount).GetNormalized();
 	}
 
 	// propagate to non-solved bones (requires storage in root to tip order)
@@ -741,7 +739,7 @@ bool FPBIKSolver::InitConstraints()
 		}
 
 		FRigidBody* Body = BodyBone->Body;
-		TSharedPtr<FPinConstraint> Constraint = MakeShared<FPinConstraint>(Body, Effector.Position, Effector.Rotation, false);
+		TSharedPtr<FPinConstraint> Constraint = MakeShared<FPinConstraint>(Body, Effector.Position, Effector.Rotation);
 		Effector.Pin = Constraint;
 		Body->Effector = &Effector;
 		Body->Pin = Constraint.Get();
@@ -762,17 +760,6 @@ bool FPBIKSolver::InitConstraints()
 		Constraints.Add(Constraint);
 	}
 
-	// pin root body to animated location 
-	// this constraint is by default off in solver settings
-	// this constraint is added last to ensure it is fully solved at the end of the iteration loop
-	if (!SolverRoot->Body->Effector) // only add if user hasn't added their own root effector
-	{
-		const TSharedPtr<FPinConstraint> RootConstraint = MakeShared<FPinConstraint>(SolverRoot->Body, SolverRoot->Position, SolverRoot->Rotation, true);
-		Constraints.Add(RootConstraint);
-		RootPin = RootConstraint;
-		SolverRoot->Body->Pin = RootConstraint.Get();
-	}
-
 	// now we can set the initial effector depths
 	UpdateEffectorDepths();
 	
@@ -829,7 +816,6 @@ void FPBIKSolver::Reset()
 	
 	bReadyToSimulate = false;
 	SolverRoot = nullptr;
-	RootPin = nullptr;
 	Bodies.Empty();
 	Bones.Empty();
 	Constraints.Empty();
diff --git a/Engine/Plugins/Experimental/FullBodyIK/Source/PBIK/Public/Core/PBIKBody.h b/Engine/Plugins/Experimental/FullBodyIK/Source/PBIK/Public/Core/PBIKBody.h
--- a/Engine/Plugins/Experimental/FullBodyIK/Source/PBIK/Public/Core/PBIKBody.h
+++ b/Engine/Plugins/Experimental/FullBodyIK/Source/PBIK/Public/Core/PBIKBody.h
@@ -118,7 +118,7 @@ public:
 
 	void ApplyPushToRotateBody(const FVector& Push, const FVector& Offset);
 	
-	void ApplyPushToPosition(const FVector& Push);
+	void ApplyPositionDelta(const FVector& InDeltaP);
 
 	void ApplyRotationDelta(const FQuat& InDelta);
 };
diff --git a/Engine/Plugins/Experimental/FullBodyIK/Source/PBIK/Public/Core/PBIKConstraint.h b/Engine/Plugins/Experimental/FullBodyIK/Source/PBIK/Public/Core/PBIKConstraint.h
--- a/Engine/Plugins/Experimental/FullBodyIK/Source/PBIK/Public/Core/PBIKConstraint.h
+++ b/Engine/Plugins/Experimental/FullBodyIK/Source/PBIK/Public/Core/PBIKConstraint.h
@@ -100,15 +100,13 @@ private:
 	FVector PinPointLocalToA;
 	float Alpha = 1.0;
 	FQuat ARotLocalToPin;
-	bool bLockRotation = false;
 
 public:
 	
 	FPinConstraint(
 		FRigidBody* InBody,
 		const FVector& InPinPositionOrig,
-		const FQuat& InPinRotationOrig,
-		const bool bInLockRotation);
+		const FQuat& InPinRotationOrig);
 
 	virtual ~FPinConstraint() {};
 
@@ -116,8 +114,6 @@ public:
 
 	void SetGoal(const FVector& InGoalPosition, const FQuat& InGoalRotation, const float InAlpha);
 
-	void EnableInCurrentState();
-
 private:
 
 	FVector GetPositionCorrection(FVector& OutBodyToPinPoint) const;
diff --git a/Engine/Plugins/Experimental/FullBodyIK/Source/PBIK/Public/Core/PBIKSolver.h b/Engine/Plugins/Experimental/FullBodyIK/Source/PBIK/Public/Core/PBIKSolver.h
--- a/Engine/Plugins/Experimental/FullBodyIK/Source/PBIK/Public/Core/PBIKSolver.h
+++ b/Engine/Plugins/Experimental/FullBodyIK/Source/PBIK/Public/Core/PBIKSolver.h
@@ -296,7 +296,6 @@ private:
 		FVector& OutCurrentCentroid);
 	
 	PBIK::FBone* SolverRoot = nullptr;
-	TWeakPtr<PBIK::FPinConstraint> RootPin = nullptr;
 	TArray<PBIK::FBone> Bones;
 	TArray<PBIK::FRigidBody> Bodies;
 	TArray<TSharedPtr<PBIK::FConstraint>> Constraints;
