Fixes static analysis issue in FBIK solver.

diff --git a/Engine/Plugins/Experimental/FullBodyIK/Source/PBIK/Private/Core/PBIKSolver.cpp b/Engine/Plugins/Experimental/FullBodyIK/Source/PBIK/Private/Core/PBIKSolver.cpp
--- a/Engine/Plugins/Experimental/FullBodyIK/Source/PBIK/Private/Core/PBIKSolver.cpp
+++ b/Engine/Plugins/Experimental/FullBodyIK/Source/PBIK/Private/Core/PBIKSolver.cpp
@@ -287,10 +287,19 @@ void FPBIKSolver::UpdateBonesFromBodies()
 		}
 
 		// if effector is between other effectors, it will have an associated rigid body
-		// so leave transform where body ended up after solve, otherwise snap position back to FK location
+		// so leave transform where body ended up after solve (applied above)
+		// otherwise, snap position back to FK location
 		if (!Bone->Body)
 		{
-			Bone->Position = Bone->Parent->Position + Bone->Parent->Rotation * Bone->LocalPositionFromInput;
+			if (Bone->Parent)
+			{
+				Bone->Position = Bone->Parent->Position + Bone->Parent->Rotation * Bone->LocalPositionFromInput;
+			}
+			else
+			{
+				// effector was on the solver root
+				Bone->Position = Effector.Position;
+			}
 		}
 
 		// optionally pin rotation to that of effector
