ControlRig.RigUnits.FullBodyIKNode.BoneSettings
#rnx

diff --git a/Engine/Plugins/Experimental/FullBodyIK/Source/PBIK/PBIK.Build.cs b/Engine/Plugins/Experimental/FullBodyIK/Source/PBIK/PBIK.Build.cs
--- a/Engine/Plugins/Experimental/FullBodyIK/Source/PBIK/PBIK.Build.cs
+++ b/Engine/Plugins/Experimental/FullBodyIK/Source/PBIK/PBIK.Build.cs
@@ -13,7 +13,7 @@ public class PBIK : ModuleRules
 			{
 				"Core",
 				"RigVM",
-				"ControlRig"
+				"ControlRig",
 			}
 			);
 
@@ -24,6 +24,16 @@ public class PBIK : ModuleRules
 			}
 			);
 
-		
+		if (Target.bBuildEditor == true)
+		{
+			PublicDependencyModuleNames.Add("ControlRigDeveloper");
+			PublicDependencyModuleNames.Add("ControlRigEditor");
+
+			PrivateDependencyModuleNames.Add("Engine");
+			PrivateDependencyModuleNames.Add("AssetTools");
+			PrivateDependencyModuleNames.Add("UnrealEd");
+			PrivateDependencyModuleNames.Add("RigVMDeveloper");
+			PrivateDependencyModuleNames.Add("RigVMEditor");
+		}
 	}
 }
diff --git a/Engine/Plugins/Experimental/FullBodyIK/Source/PBIK/Private/Tests/ControlRigEditorFullBodyIK.spec.cpp b/Engine/Plugins/Experimental/FullBodyIK/Source/PBIK/Private/Tests/ControlRigEditorFullBodyIK.spec.cpp
new file mode 100644
--- /dev/null
+++ b/Engine/Plugins/Experimental/FullBodyIK/Source/PBIK/Private/Tests/ControlRigEditorFullBodyIK.spec.cpp
@@ -0,0 +1,236 @@
+// Copyright Epic Games, Inc. All Rights Reserved.
+
+#if WITH_DEV_AUTOMATION_TESTS
+
+#if WITH_EDITOR
+
+#include "Misc/AutomationTest.h"
+#include "Editor.h"
+#include "Subsystems/EditorAssetSubsystem.h"
+#include "Subsystems/AssetEditorSubsystem.h"
+#include "ControlRigBlueprintFactory.h"
+#include "ControlRigBlueprint.h"
+#include "Rigs/RigHierarchyController.h"
+#include "Rigs/RigHierarchy.h"
+#include "Graph/ControlRigGraph.h"
+#include "RigVMModel/RigVMController.h"
+#include "ObjectTools.h"
+#include "Units/Hierarchy/RigUnit_GetTransform.h"
+#include "RigUnit_PBIK.h"
+#include "IControlRigEditor.h"
+#include "SAdvancedTransformInputBox.h"
+#include "Core/PBIKSolver.h"
+
+
+BEGIN_DEFINE_SPEC(
+	FControlRigEditorFullBodyIKNodeSpec, "Animation.Editor",
+	EAutomationTestFlags::EditorContext | EAutomationTestFlags::EngineFilter)
+
+	UControlRigBlueprint* ControlRigAsset = nullptr;
+	UControlRigGraph* RigGraph = nullptr;
+	URigVMController* VMController = nullptr;
+	
+	TStrongObjectPtr<URigVMUnitNode> ControlTransformNode;
+	TStrongObjectPtr<URigVMUnitNode> FullBodyIKNode;
+	
+	IControlRigEditor* ControlRigEditor = nullptr;
+	FRigElementKey ControlKey;
+
+	void AddControlElement();
+	const FRigElementKey GetBoneKey(const FName&) const;
+	void CreateAndLinkNodes();
+	bool ApplyTransformToControl(double, double) const;
+	double GetBoneYRotation(const FName&) const;
+
+END_DEFINE_SPEC(FControlRigEditorFullBodyIKNodeSpec)
+
+void FControlRigEditorFullBodyIKNodeSpec::Define()
+{
+	Describe("ControlRigEditorFullBodyIKNodeTest", [this]()
+		{
+			BeforeEach([this]()
+				{
+					const FString SkeletalMeshName = "chainCubes";
+					const FString SkeletalMeshPath = FString::Printf(TEXT("/Game/Tests/Physics/RBAN/%s"), *SkeletalMeshName);
+					const FAssetData AssetData = GEditor->GetEditorSubsystem<UEditorAssetSubsystem>()->FindAssetData(*SkeletalMeshPath);
+					UObject* SkeletalMesh = nullptr;
+					if (AssetData.IsValid())
+					{
+						SkeletalMesh = AssetData.GetAsset();
+					}
+					TestNotNull(FString::Printf(TEXT("Asset '%s' loaded successfully"), *SkeletalMeshName), SkeletalMesh);
+			
+					ControlRigAsset = UControlRigBlueprintFactory::CreateControlRigFromSkeletalMeshOrSkeleton(SkeletalMesh);
+					TestNotNull(TEXT("Control Rig Asset created successfully from Skeletal Mesh"), ControlRigAsset);
+					
+					if (UEdGraph* EdGraph = ControlRigAsset->GetEdGraph(ControlRigAsset->GetModel()))
+					{
+						RigGraph = CastChecked<UControlRigGraph>(EdGraph);
+						VMController = ControlRigAsset->GetController(RigGraph);
+					}
+					
+					if(UAssetEditorSubsystem* AssetEditorSubsystem = GEditor->GetEditorSubsystem<UAssetEditorSubsystem>())
+					{
+						const bool bShowProgressWindow = false;
+						if (AssetEditorSubsystem->OpenEditorForAsset(ControlRigAsset, EToolkitMode::Standalone, TSharedPtr<IToolkitHost>(), bShowProgressWindow))
+						{
+							const bool bFocusIfOpen = true;
+							if (IAssetEditorInstance* ControlRigEditorInstance = AssetEditorSubsystem->FindEditorForAsset(ControlRigAsset, bFocusIfOpen))
+							{
+								ControlRigEditor = static_cast<IControlRigEditor*>(ControlRigEditorInstance);
+							}
+						}
+					}
+
+					AddControlElement();
+					CreateAndLinkNodes();
+					ApplyTransformToControl(1000,0);
+				});
+
+			It("BoneSettingsEPBIKLimitTypeLimitedTest", [this]()
+				{
+					double MaxY = 45.0;
+
+					VMController->SetPinDefaultValue(FString::Printf(TEXT("%s.BoneSettings.0.Y"), *FullBodyIKNode->GetName()), "EPBIKLimitType::Limited");
+					VMController->SetPinDefaultValue(FString::Printf(TEXT("%s.BoneSettings.0.MaxY"), *FullBodyIKNode->GetName()), FString::SanitizeFloat(MaxY));
+					ApplyTransformToControl(1000, 6000);
+					double YRotationLimited = GetBoneYRotation("jointA");
+
+					VMController->SetPinDefaultValue(FString::Printf(TEXT("%s.BoneSettings.0.Y"), *FullBodyIKNode->GetName()), "EPBIKLimitType::Free");
+					ApplyTransformToControl(1000, 6000);
+					double YRotationFree = GetBoneYRotation("jointA");
+
+					TestTrue(TEXT("Limit Bone rotation along the Y-axis to the specified MaxY value"), YRotationFree > MaxY && YRotationLimited > 0.0 && (YRotationLimited < MaxY || FMath::IsNearlyEqual(YRotationLimited, MaxY)));
+				});
+
+			It("BoneSettingsEPBIKLimitTypeLockedTest", [this]()
+				{
+					VMController->SetPinDefaultValue(FString::Printf(TEXT("%s.BoneSettings.0.Y"), *FullBodyIKNode->GetName()), "EPBIKLimitType::Locked");
+					ApplyTransformToControl(1000, 300);
+					double YRotationLocked = GetBoneYRotation("jointA");
+
+					VMController->SetPinDefaultValue(FString::Printf(TEXT("%s.BoneSettings.0.Y"), *FullBodyIKNode->GetName()), "EPBIKLimitType::Free");
+					ApplyTransformToControl(1000, 300);
+					double YRotationFree = GetBoneYRotation("jointA");
+
+					TestTrue(TEXT("Bone rotation remains unchanged"), FMath::IsNearlyEqual(0.0, YRotationLocked) && YRotationFree > 0.0);
+				});
+
+			AfterEach([this]()
+				{
+					if (ControlRigAsset)
+					{
+						GEditor->GetEditorSubsystem<UAssetEditorSubsystem>()->CloseAllEditorsForAsset(ControlRigAsset);
+						ObjectTools::DeleteObjectsUnchecked({ ControlRigAsset });
+					}
+				});
+		});
+}
+
+void FControlRigEditorFullBodyIKNodeSpec::AddControlElement()
+{
+	if (URigHierarchyController* HierarchyController = ControlRigAsset->GetHierarchyController())
+	{
+		ControlKey = HierarchyController->AddControl(TEXT("Control"), FRigElementKey(), FRigControlSettings(), FRigControlValue::Make(FTransform::Identity));
+	}
+}
+
+const FRigElementKey FControlRigEditorFullBodyIKNodeSpec::GetBoneKey(const FName& Name) const
+{
+	if (ControlRigAsset)
+	{
+		TArray<FRigElementKey> BoneKeys = ControlRigAsset->GetHierarchy()->GetBoneKeys();
+		FRigElementKey* BoneKey = BoneKeys.FindByPredicate([Name](const FRigElementKey& Key)
+			{
+				return Key.Name.IsEqual(Name);
+			});
+
+		if (BoneKey)
+		{
+			return *BoneKey;
+
+		}
+	}
+	return FRigElementKey();
+}
+
+void FControlRigEditorFullBodyIKNodeSpec::CreateAndLinkNodes()
+{
+	if (VMController)
+	{
+		FRigUnit_GetTransform GetTransform;
+		GetTransform.Item = ControlKey;
+		ControlTransformNode.Reset(VMController->AddUnitNode(GetTransform));
+
+		FRigUnit_PBIK FullBodyIK;
+		FullBodyIK.Effectors.AddDefaulted();
+		FullBodyIK.BoneSettings.AddDefaulted();
+		FullBodyIKNode.Reset(VMController->AddUnitNode(FullBodyIK));
+
+		if (RigGraph && ControlTransformNode.IsValid() && FullBodyIKNode.IsValid())
+		{
+			const URigVMGraph* RigVMGraph = RigGraph->GetModel();
+			const TArray<URigVMNode*>& Nodes = RigVMGraph->GetNodes();
+
+			if (Nodes.Num() > 0)
+			{
+				VMController->AddLink(
+					FString::Printf(TEXT("%s.ExecuteContext"), *Nodes[0]->GetName()),
+					FString::Printf(TEXT("%s.ExecuteContext"), *FullBodyIKNode->GetName())
+				);
+			}
+
+			VMController->AddLink(
+				FString::Printf(TEXT("%s.Transform"), *ControlTransformNode->GetName()),
+				FString::Printf(TEXT("%s.Effectors.0.Transform"), *FullBodyIKNode->GetName())
+			);
+		}
+
+		VMController->SetPinDefaultValue(FString::Printf(TEXT("%s.Root"), *FullBodyIKNode->GetName()), "root");
+		VMController->SetPinDefaultValue(FString::Printf(TEXT("%s.Effectors.0.Bone"), *FullBodyIKNode->GetName()), "jointE");
+		VMController->SetPinDefaultValue(FString::Printf(TEXT("%s.BoneSettings.0.Bone"), *FullBodyIKNode->GetName()), "jointA");
+		VMController->SetPinDefaultValue(FString::Printf(TEXT("%s.BoneSettings.0.X"), *FullBodyIKNode->GetName()), "EPBIKLimitType::Locked");
+		VMController->SetPinDefaultValue(FString::Printf(TEXT("%s.BoneSettings.0.Z"), *FullBodyIKNode->GetName()), "EPBIKLimitType::Locked");
+		VMController->SetPinDefaultValue(FString::Printf(TEXT("%s.Settings.RootBehavior"), *FullBodyIKNode->GetName()), "EPBIKRootBehavior::PinToInput");
+	}
+}
+
+bool FControlRigEditorFullBodyIKNodeSpec::ApplyTransformToControl(double XValue, double YValue) const
+{
+
+	UControlRig* ControlRig = Cast<UControlRig>(ControlRigEditor->GetRigVMHost());
+
+	if (ControlRig && ControlKey.IsValid())
+	{
+		ControlRig->GetHierarchy()->Modify();
+
+		FTransform ControlLocalTransform = ControlRig->GetControlLocalTransform(ControlKey.Name);
+		auto Location = ControlLocalTransform.GetLocation();
+		Location.X = XValue;
+		Location.Y = YValue;
+		ControlLocalTransform.SetLocation(Location);
+
+		ControlRig->Evaluate_AnyThread();
+		ControlRig->SetControlLocalTransform(ControlKey.Name, ControlLocalTransform, true, FRigControlModifiedContext(), false, false);
+		ControlRig->Evaluate_AnyThread();
+
+		return true;
+	}
+
+	return false;
+}
+
+double FControlRigEditorFullBodyIKNodeSpec::GetBoneYRotation(const FName& Name) const
+{
+	if (UControlRig* ControlRig = Cast<UControlRig>(ControlRigEditor->GetRigVMHost()))
+	{
+		FRigElementKey BoneKey = GetBoneKey(Name);
+		FTransform Transform = ControlRig->GetHierarchy()->GetLocalTransform(BoneKey);
+		return Transform.Rotator().Pitch;
+	}
+	return 0.0;
+}
+
+#endif // WITH_EDITOR
+
+#endif // WITH_DEV_AUTOMATION_TESTS
\ No newline at end of file
