FBIK core solver fixes.
- preferred angles now consider local offset in ref pose
- chain depth updates correctly

diff --git a/Engine/Plugins/Experimental/FullBodyIK/Source/PBIK/Private/Core/PBIKSolver.cpp b/Engine/Plugins/Experimental/FullBodyIK/Source/PBIK/Private/Core/PBIKSolver.cpp
--- a/Engine/Plugins/Experimental/FullBodyIK/Source/PBIK/Private/Core/PBIKSolver.cpp
+++ b/Engine/Plugins/Experimental/FullBodyIK/Source/PBIK/Private/Core/PBIKSolver.cpp
@@ -148,8 +148,10 @@ namespace PBIK
 			if (Body->J.bUsePreferredAngles)
 			{
 				const FRigidBody* CurrentParent = Body->GetParentBody();
-				FQuat LocalRotation = CurrentParent ? CurrentParent->Rotation.Inverse() * Body->Rotation : Body->Rotation;
-				const FRotator AnglesFromInput = LocalRotation.Rotator();
+				FQuat CurrentLocalRotation = CurrentParent ? CurrentParent->Rotation.Inverse() * Body->Rotation : Body->Rotation;
+				FQuat InitialLocalRotation = CurrentParent ? CurrentParent->InitialRotation.Inverse() * Body->InitialRotation : Body->InitialRotation;
+				FQuat LocalOffsetRotation = CurrentLocalRotation * InitialLocalRotation.Inverse();
+				const FRotator AnglesFromInput = LocalOffsetRotation.Rotator();
 				const FRotator PreferredAngles = Body->J.PreferredAngles * SquashPercent;
 
 				auto CalcDeltaAngle = [](double PreferredAngle, double InputAngle) -> double
@@ -231,7 +233,8 @@ void FPBIKSolver::Solve(const FPBIKSolverSettings& Settings)
 	}
 
 	// lazily updates effector chain depths (IFF settings change at runtime)
-	UpdateEffectorDepths();
+	constexpr bool bForceUpdate = false;
+	UpdateEffectorDepths(bForceUpdate);
 
 	// blend effectors by Alpha, update pin goals and update effector dist to root
 	for (FEffector& Effector : Effectors)
@@ -770,12 +773,13 @@ bool FPBIKSolver::InitConstraints()
 	}
 
 	// now we can set the initial effector depths
-	UpdateEffectorDepths();
+	constexpr bool bForceUpdate = true;
+	UpdateEffectorDepths(bForceUpdate);
 	
 	return true;
 }
 
-void FPBIKSolver::UpdateEffectorDepths()
+void FPBIKSolver::UpdateEffectorDepths(bool bForceUpdate)
 {
 	using PBIK::FEffector;
 	using PBIK::FBone;
@@ -787,7 +791,7 @@ void FPBIKSolver::UpdateEffectorDepths()
 	bHasSubChains = false;
 	for (const FEffector& Effector : Effectors)
 	{
-		if (Effector.ChainDepthInitializedWith != Effector.Settings.ChainDepth)
+		if (bForceUpdate || (Effector.ChainDepthInitializedWith != Effector.Settings.ChainDepth))
 		{
 			bIsEffectorDepthDirty = true;
 		}
diff --git a/Engine/Plugins/Experimental/FullBodyIK/Source/PBIK/Public/Core/PBIKSolver.h b/Engine/Plugins/Experimental/FullBodyIK/Source/PBIK/Public/Core/PBIKSolver.h
--- a/Engine/Plugins/Experimental/FullBodyIK/Source/PBIK/Public/Core/PBIKSolver.h
+++ b/Engine/Plugins/Experimental/FullBodyIK/Source/PBIK/Public/Core/PBIKSolver.h
@@ -268,7 +268,7 @@ private:
 
 	bool InitConstraints();
 
-	void UpdateEffectorDepths();
+	void UpdateEffectorDepths(bool bForceUpdate);
 
 	void AddBodyForBone(PBIK::FBone* Bone);
 
