Fixes FBIK bug, solver root body not always locking during pre solve iterations.

diff --git a/Engine/Plugins/Experimental/FullBodyIK/Source/PBIK/Private/Core/PBIKBody.cpp b/Engine/Plugins/Experimental/FullBodyIK/Source/PBIK/Private/Core/PBIKBody.cpp
--- a/Engine/Plugins/Experimental/FullBodyIK/Source/PBIK/Private/Core/PBIKBody.cpp
+++ b/Engine/Plugins/Experimental/FullBodyIK/Source/PBIK/Private/Core/PBIKBody.cpp
@@ -119,14 +119,14 @@ FRigidBody* FRigidBody::GetParentBody() const
 	return nullptr;
 }
 
-bool FRigidBody::IsAllowedToRotate() const
+bool FRigidBody::IsFullyLocked() const
 {
-	return !bIsLockedBySubSolve && InvMass > SMALL_NUMBER;
+	return bIsLockedBySubSolve || InvMass <= SMALL_NUMBER;
 }
 
 void FRigidBody::ApplyPushToRotateBody(const FVector& Push, const FVector& Offset)
 {
-	if (!IsAllowedToRotate())
+	if (IsFullyLocked())
 	{
 		return; // rotation of this body is disabled
 	}
@@ -139,12 +139,17 @@ void FRigidBody::ApplyPushToRotateBody(const FVector& Push, const FVector& Offse
 
 void FRigidBody::ApplyPositionDelta(const FVector& DeltaP)
 {
+	if (IsFullyLocked())
+	{
+		return; // translation of this body is disabled
+	}
+	
 	Position += DeltaP * (1.0f - J.PositionStiffness) * SolverSettings->OverRelaxation;
 }
 
 void FRigidBody::ApplyRotationDelta(const FQuat& DeltaQ)
 {
-	if (!IsAllowedToRotate())
+	if (IsFullyLocked())
 	{
 		return; // rotation of this body is disabled
 	}
diff --git a/Engine/Plugins/Experimental/FullBodyIK/Source/PBIK/Public/Core/PBIKBody.h b/Engine/Plugins/Experimental/FullBodyIK/Source/PBIK/Public/Core/PBIKBody.h
--- a/Engine/Plugins/Experimental/FullBodyIK/Source/PBIK/Public/Core/PBIKBody.h
+++ b/Engine/Plugins/Experimental/FullBodyIK/Source/PBIK/Public/Core/PBIKBody.h
@@ -114,7 +114,7 @@ public:
 
 	FRigidBody* GetParentBody() const;
 
-	bool IsAllowedToRotate() const;
+	bool IsFullyLocked() const;
 
 	void ApplyPushToRotateBody(const FVector& Push, const FVector& Offset);
 	
