Tests migration to EngineTest project
#rnx

diff --git a/Engine/Plugins/Animation/ControlRig/Source/ControlRigEditor/Private/Tests/ControlRigEditor.spec.cpp b/Engine/Plugins/Animation/ControlRig/Source/ControlRigEditor/Private/Tests/ControlRigEditor.spec.cpp
deleted file mode 100644
diff --git a/Engine/Plugins/Animation/ControlRig/Source/ControlRigEditor/Private/Tests/ControlRigEditorlTransform.spec.cpp b/Engine/Plugins/Animation/ControlRig/Source/ControlRigEditor/Private/Tests/ControlRigEditorlTransform.spec.cpp
deleted file mode 100644
diff --git a/Engine/Plugins/Animation/IKRig/Source/IKRigEditor/Private/RetargetEditor/Tests/RetargetEditor.spec.cpp b/Engine/Plugins/Animation/IKRig/Source/IKRigEditor/Private/RetargetEditor/Tests/RetargetEditor.spec.cpp
deleted file mode 100644
diff --git a/Engine/Plugins/Experimental/FullBodyIK/Source/PBIK/Private/Tests/ControlRigEditorFullBodyIK.spec.cpp b/Engine/Plugins/Experimental/FullBodyIK/Source/PBIK/Private/Tests/ControlRigEditorFullBodyIK.spec.cpp
deleted file mode 100644
diff --git a/Engine/Plugins/Experimental/FullBodyIK/Source/PBIK/Public/RigUnit_PBIK.h b/Engine/Plugins/Experimental/FullBodyIK/Source/PBIK/Public/RigUnit_PBIK.h
--- a/Engine/Plugins/Experimental/FullBodyIK/Source/PBIK/Public/RigUnit_PBIK.h
+++ b/Engine/Plugins/Experimental/FullBodyIK/Source/PBIK/Public/RigUnit_PBIK.h
@@ -139,7 +139,7 @@ struct FPBIKWorkData
  * Based on a Position Based solver at core, this node can solve multi chains within a root using multi effectors
  */
 USTRUCT(meta=(DisplayName="Full Body IK", Category="Hierarchy", Keywords="FBIK, Position Based, PBIK, IK, Full Body, Multi, Effector, N-Chain, FB, HIK, HumanIK", NodeColor="0 1 1"))
-struct FRigUnit_PBIK : public FRigUnit_HighlevelBaseMutable
+struct PBIK_API FRigUnit_PBIK : public FRigUnit_HighlevelBaseMutable
 {
 	GENERATED_BODY()
 
