[FBIK]
Fixed issue, pulling effectors far causes solve instability.
Fixed issue, rotation limits not fully respected.
Fixed issue, rotation Y min / max angles were backwards.

diff --git a/Engine/Plugins/Experimental/FullBodyIK/Source/PBIK/Private/Core/PBIKBody.cpp b/Engine/Plugins/Experimental/FullBodyIK/Source/PBIK/Private/Core/PBIKBody.cpp
--- a/Engine/Plugins/Experimental/FullBodyIK/Source/PBIK/Private/Core/PBIKBody.cpp
+++ b/Engine/Plugins/Experimental/FullBodyIK/Source/PBIK/Private/Core/PBIKBody.cpp
@@ -119,14 +119,19 @@ FRigidBody* FRigidBody::GetParentBody() const
 	return nullptr;
 }
 
-bool FRigidBody::IsFullyLocked() const
+bool FRigidBody::IsPositionLocked() const
 {
-	return bIsLockedBySubSolve || InvMass <= SMALL_NUMBER;
+	return bIsLockedBySubSolve || InvMass <= SMALL_NUMBER || FMath::IsNearlyEqual(J.PositionStiffness, 1.0f);
+}
+
+bool FRigidBody::IsRotationLocked() const
+{
+	return bIsLockedBySubSolve || InvMass <= SMALL_NUMBER || FMath::IsNearlyEqual(J.RotationStiffness, 1.0f);
 }
 
 void FRigidBody::ApplyPushToRotateBody(const FVector& Push, const FVector& Offset)
 {
-	if (IsFullyLocked())
+	if (IsRotationLocked())
 	{
 		return; // rotation of this body is disabled
 	}
@@ -139,7 +144,7 @@ void FRigidBody::ApplyPushToRotateBody(const FVector& Push, const FVector& Offse
 
 void FRigidBody::ApplyPositionDelta(const FVector& DeltaP)
 {
-	if (IsFullyLocked())
+	if (IsPositionLocked())
 	{
 		return; // translation of this body is disabled
 	}
@@ -149,7 +154,7 @@ void FRigidBody::ApplyPositionDelta(const FVector& DeltaP)
 
 void FRigidBody::ApplyRotationDelta(const FQuat& DeltaQ)
 {
-	if (IsFullyLocked())
+	if (IsRotationLocked())
 	{
 		return; // rotation of this body is disabled
 	}
diff --git a/Engine/Plugins/Experimental/FullBodyIK/Source/PBIK/Private/Core/PBIKConstraint.cpp b/Engine/Plugins/Experimental/FullBodyIK/Source/PBIK/Private/Core/PBIKConstraint.cpp
--- a/Engine/Plugins/Experimental/FullBodyIK/Source/PBIK/Private/Core/PBIKConstraint.cpp
+++ b/Engine/Plugins/Experimental/FullBodyIK/Source/PBIK/Private/Core/PBIKConstraint.cpp
@@ -21,6 +21,8 @@ FJointConstraint::FJointConstraint(FRigidBody* InA, FRigidBody* InB)
 	ZOrig = B->InitialRotation * FVector(0, 0, 1);
 
 	AngleX = AngleY = AngleZ = 0;
+
+	bInFinalPass = false;
 }
 	
 void FJointConstraint::Solve(const FPBIKSolverSettings& Settings)
@@ -78,6 +80,13 @@ void FJointConstraint::RemoveStretch(const float Percent)
 	B->Position -= Correction * Percent;
 }
 
+void FJointConstraint::FinalPass()
+{
+	bInFinalPass = true;
+	UpdateJointLimits();
+	bInFinalPass = false;
+}
+
 void FJointConstraint::UpdateFromInputs()
 {
 	const FVector PinPoint = B->Bone->Position;
@@ -118,9 +127,10 @@ void FJointConstraint::UpdateJointLimits()
 	J.MaxZ = J.MaxZ < J.MinZ ? J.MinZ + 1 : J.MaxZ;
 
 	// determine which axes are explicitly or implicitly locked (limited with very small allowable angle)
-	const bool bLockX = (J.X == ELimitType::Locked) || (J.X == ELimitType::Limited && ((J.MaxX - J.MinX) < 2.0f));
-	const bool bLockY = (J.Y == ELimitType::Locked) || (J.Y == ELimitType::Limited && ((J.MaxY - J.MinY) < 2.0f));
-	const bool bLockZ = (J.Z == ELimitType::Locked) || (J.Z == ELimitType::Limited && ((J.MaxZ - J.MinZ) < 2.0f));
+	constexpr float LockThresholdAngle = 2.0f; // degrees
+	const bool bLockX = (J.X == ELimitType::Locked) || (J.X == ELimitType::Limited && ((J.MaxX - J.MinX) < LockThresholdAngle));
+	const bool bLockY = (J.Y == ELimitType::Locked) || (J.Y == ELimitType::Limited && ((J.MaxY - J.MinY) < LockThresholdAngle));
+	const bool bLockZ = (J.Z == ELimitType::Locked) || (J.Z == ELimitType::Limited && ((J.MaxZ - J.MinZ) < LockThresholdAngle));
 
 	// when locked on THREE axes, consider the bone FIXED
 	if (bLockX && bLockY && bLockZ)
@@ -128,7 +138,13 @@ void FJointConstraint::UpdateJointLimits()
 		const FQuat OrigOffset = A->InitialRotation * B->InitialRotation.Inverse();
 		const FQuat CurrentOffset = A->Rotation * B->Rotation.Inverse();
 		const FQuat DeltaOffset = OrigOffset * CurrentOffset.Inverse();
-		const FQuat DeltaQ = FQuat(DeltaOffset.X, DeltaOffset.Y, DeltaOffset.Z, 0.0f) * 2.0f;
+		if (bInFinalPass)
+		{
+			B->Rotation = DeltaOffset.Inverse() * B->Rotation;
+			return;
+		}
+		
+		const FQuat DeltaQ = FQuat(DeltaOffset.X, DeltaOffset.Y, DeltaOffset.Z, 0.0f);
 		ApplyRotationCorrection(DeltaQ);
 		return; // fixed bones cannot move at all
 	}
@@ -171,7 +187,7 @@ void FJointConstraint::UpdateJointLimits()
 
 	if (bLimitY)
 	{
-		RotateWithinLimits(J.MinY, J.MaxY, AngleY, YA, ZBProjOnY, ZA);
+		RotateWithinLimits(J.MinY, J.MaxY, AngleY, -YA, ZBProjOnY, ZA);
 	}
 
 	if (bLimitZ)
@@ -182,6 +198,13 @@ void FJointConstraint::UpdateJointLimits()
 
 void FJointConstraint::RotateToAlignAxes(const FVector& AxisA, const FVector& AxisB) const
 {
+	if (bInFinalPass)
+	{
+		const FQuat DeltaQ = FQuat::FindBetweenVectors(AxisB, AxisA);
+		B->Rotation = DeltaQ * B->Rotation;
+		return;
+	}
+	
 	const FVector ACrossB = FVector::CrossProduct(AxisA, AxisB);
 	FVector Axis;
 	float Magnitude;
@@ -258,7 +281,7 @@ void FJointConstraint::DecomposeRotationAngles()
 	YBProjOnZ = FVector::VectorPlaneProject(YB, ZA).GetSafeNormal();
 
 	AngleX = SignedAngleBetweenNormals(ZA, ZBProjOnX, XA);
-	AngleY = SignedAngleBetweenNormals(ZA, ZBProjOnY, YA);
+	AngleY = -SignedAngleBetweenNormals(ZA, ZBProjOnY, YA);
 	AngleZ = SignedAngleBetweenNormals(YA, YBProjOnZ, ZA);
 }
 
@@ -273,7 +296,7 @@ float FJointConstraint::SignedAngleBetweenNormals(
 	const float Dot = FVector::DotProduct(Cross, Axis);
 	return Dot >= 0 ? Angle : -Angle;
 }
-	
+
 FPinConstraint::FPinConstraint(
 	FRigidBody* InBody,
 	const FVector& InPinPositionOrig,
diff --git a/Engine/Plugins/Experimental/FullBodyIK/Source/PBIK/Private/Core/PBIKSolver.cpp b/Engine/Plugins/Experimental/FullBodyIK/Source/PBIK/Private/Core/PBIKSolver.cpp
--- a/Engine/Plugins/Experimental/FullBodyIK/Source/PBIK/Private/Core/PBIKSolver.cpp
+++ b/Engine/Plugins/Experimental/FullBodyIK/Source/PBIK/Private/Core/PBIKSolver.cpp
@@ -64,11 +64,21 @@ namespace PBIK
 		ChainDepthInitializedWith = Settings.ChainDepth;
 	}
 
-	void FEffector::UpdateFromInputs(const FBone& SolverRoot)
+	void FEffector::UpdateFromInputs(const FBone& InSolverRoot, const bool bAllowStretch)
 	{
-		// blend effector transform by alpha and set pin goal transform
+		// blend effector transform by alpha
 		Position = FMath::Lerp(PositionOrig, PositionGoal, Settings.PositionAlpha);
 		Rotation = FMath::Lerp(RotationOrig, RotationGoal, Settings.RotationAlpha);
+
+		// if the solver root is pinned, we clamp the effector position so that it cannot pull excessively,
+		// this prevents extreme constraint forces that cause solve instability
+		if (FMath::IsNearlyZero(InSolverRoot.Body->InvMass) && !bAllowStretch)
+		{
+			const float DistToSolveRootAlongBones = CalculateDistanceToSolveRootAlongBones();
+			Position = InSolverRoot.Position + (Position - InSolverRoot.Position).GetClampedToMaxSize(DistToSolveRootAlongBones);
+		}
+
+		// set pin goal transform
 		Pin.Pin()->SetGoal(Position, Rotation, Settings.StrengthAlpha);
 		
 		// update length of chain this effector controls in the input pose
@@ -103,6 +113,18 @@ namespace PBIK
 		return 0.f;
 	}
 
+	float FEffector::CalculateDistanceToSolveRootAlongBones() const
+	{
+		FBone* CurrentBone = Bone;
+		float DistanceTotal = 0.0f; 
+		while (!CurrentBone->bIsSolverRoot)
+		{
+			DistanceTotal += CurrentBone->Length;
+			CurrentBone = CurrentBone->Parent;
+		}
+		return DistanceTotal;
+	}
+
 	void FEffector::ApplyPreferredAngles() const
 	{
 		// optionally apply a preferred angle to give solver a hint which direction to favor
@@ -189,7 +211,7 @@ namespace PBIK
 	}
 } // namespace
 
-void FPBIKSolver::Solve(const FPBIKSolverSettings& Settings)
+void FPBIKSolver::Solve(const FPBIKSolverSettings& InSolverSettings)
 {
 	SCOPE_CYCLE_COUNTER(STAT_PBIK_Solve);
 
@@ -215,7 +237,7 @@ void FPBIKSolver::Solve(const FPBIKSolverSettings& Settings)
 	// update Bodies with new bone positions from incoming pose and solver settings
 	for (FRigidBody& Body : Bodies)
 	{
-		Body.UpdateFromInputs(Settings);
+		Body.UpdateFromInputs(InSolverSettings);
 	}
 
 	// update constraints with new local pin locations relative to (possibly) translated joint positions
@@ -225,7 +247,7 @@ void FPBIKSolver::Solve(const FPBIKSolverSettings& Settings)
 	}
 
 	// optionally pin root in-place (convenience, does not require an effector)
-	if (Settings.RootBehavior == EPBIKRootBehavior::PinToInput)
+	if (InSolverSettings.RootBehavior == EPBIKRootBehavior::PinToInput)
 	{
 		// we can safely override the mass here to pin the root body
 		// if the user switches to a different root mode, the call to Body.UpdateFromInputs() will restore the correct InvMass for the body
@@ -239,11 +261,11 @@ void FPBIKSolver::Solve(const FPBIKSolverSettings& Settings)
 	// blend effectors by Alpha, update pin goals and update effector dist to root
 	for (FEffector& Effector : Effectors)
 	{
-		Effector.UpdateFromInputs(*SolverRoot);
+		Effector.UpdateFromInputs(*SolverRoot, InSolverSettings.bAllowStretch);
 	}
 
 	// do all constraint solving
-	UpdateBodies(Settings);
+	UpdateBodies(InSolverSettings);
 
 	// now that bodies are posed, update bone transforms
 	UpdateBonesFromBodies();
@@ -442,16 +464,25 @@ void FPBIKSolver::ApplyPullChainAlpha(const float GlobalPullChainAlpha)
 		int32 ChainIndex = 0;
 		const float InvChainLength = 1.0f / Effector.DistToChainRootAlongBones;
 		while (Bone && Bone->Body != Effector.ChainRootBody)
-		{		
+		{
+			// move body along with chain
+			if (!Bone->Body->IsPositionLocked())
+			{
+				// move first to the rotated chain position
+				const FVector BodyRelativeToChain = Bone->Body->Position - ChainStartNew;
+				const FVector RotatedBodyPos = ChainStartNew + ChainDeltaRotation.RotateVector(BodyRelativeToChain);
+				Bone->Body->Position = RotatedBodyPos;
+
+				// translate along chain proportional to chain stretch amount
+				const float Strength = 1.0f - (Effector.DistancesFromEffector[ChainIndex] * InvChainLength);
+				Bone->Body->Position += ChainDeltaPosition * Strength;
+			}
+
 			// rotate body along with chain
-			const FVector BodyRelativeToChain = Bone->Body->Position - ChainStartNew;
-			const FVector RotatedBodyPos = ChainStartNew + ChainDeltaRotation.RotateVector(BodyRelativeToChain);
-			Bone->Body->Position = RotatedBodyPos;
-			Bone->Body->Rotation = ChainDeltaRotation * Bone->Body->Rotation;
-
-			// move body proportional to chain stretch amount
-			const float Strength = 1.0f - (Effector.DistancesFromEffector[ChainIndex] * InvChainLength);
-			Bone->Body->Position += ChainDeltaPosition * Strength;
+			if (!Bone->Body->IsRotationLocked())
+			{
+				Bone->Body->Rotation = ChainDeltaRotation * Bone->Body->Rotation;
+			}
 			
 			Bone = Bone->Parent;
 			++ChainIndex;
@@ -476,32 +507,17 @@ void FPBIKSolver::ApplyPreferredAngles()
 	}
 }
 
-void FPBIKSolver::SolveConstraints(const FPBIKSolverSettings& Settings)
+void FPBIKSolver::SolveConstraints(const FPBIKSolverSettings& InSolverSettings)
 {
 	using PBIK::FRigidBody;
 
-	auto RunConstraintPass = [this, &Settings](int32 CurrentIteration, int32 MaxIterations)
+	auto RunConstraintPass = [this, &InSolverSettings](int32 CurrentIteration, int32 MaxIterations)
 	{
-		// solve all constraints
-		for (const TSharedPtr<PBIK::FConstraint>& Constraint : Constraints)
-		{
-			Constraint->Solve(Settings);
-		}
-		
-		// do post-pass to remove stretch
-		if (!Settings.bAllowStretch)
-		{
-			const float StretchAmount = FMath::Min(1.0f, static_cast<float>(CurrentIteration+1) / (MaxIterations / 2.0f));
-			for (int32 C = Constraints.Num() - 1; C >= 0; --C)
-			{
-				Constraints[C]->RemoveStretch(StretchAmount);
-			}
-		}
 		
 	};
 
 	// run pre-pass where we lock everything not in a sub-chain
-	if (bHasSubChains && Settings.SubIterations > 0)
+	if (bHasSubChains && InSolverSettings.SubIterations > 0)
 	{
 		// first lock all the bodies not in a sub chain
 		for (FRigidBody& Body : Bodies)
@@ -510,9 +526,9 @@ void FPBIKSolver::SolveConstraints(const FPBIKSolverSettings& Settings)
 		}
 		
 		// solve all constraints (locked constraints are skipped in this pass)
-		for (int32  I = 0; I < Settings.SubIterations; ++I)
+		for (int32  I = 0; I < InSolverSettings.SubIterations; ++I)
 		{
-			RunConstraintPass(I, Settings.SubIterations);
+			RunConstraintPass(I, InSolverSettings.SubIterations);
 		}
 		
 		// unlock all bodies so we can proceed with normal constraint solve
@@ -523,9 +539,34 @@ void FPBIKSolver::SolveConstraints(const FPBIKSolverSettings& Settings)
 	}
 
 	// run main constraint pass
-	for (int32  I = 0; I < Settings.Iterations; ++I)
+	for (int32  CurrentIteration = 0; CurrentIteration < InSolverSettings.Iterations; ++CurrentIteration)
 	{
-		RunConstraintPass(I, Settings.Iterations);
+		// solve all constraints
+		for (const TSharedPtr<PBIK::FConstraint>& Constraint : Constraints)
+		{
+			Constraint->Solve(InSolverSettings);
+		}
+
+		// on final iteration, run FinalPass
+		if (CurrentIteration == InSolverSettings.Iterations - 1)
+		{
+			// iterate in root-to-leaf order to force exact rotation limits
+			for (int32 C = Constraints.Num() - 1; C >= 0; --C)
+			{
+				Constraints[C]->FinalPass();
+			}
+		}
+		
+		// do post-pass to remove stretch
+		if (!InSolverSettings.bAllowStretch)
+		{
+			const float StretchAmount = FMath::Min(1.0f, static_cast<float>(CurrentIteration+1) / (InSolverSettings.Iterations / 2.0f));
+			// iterate in reverse order (root to leaf)
+			for (int32 C = Constraints.Num() - 1; C >= 0; --C)
+			{
+				Constraints[C]->RemoveStretch(StretchAmount);
+			}
+		}
 	}
 }
 
@@ -703,9 +744,8 @@ bool FPBIKSolver::InitBodies()
 		Body.Initialize(SolverRoot);
 	}
 
-	// sort bodies root to leaf
+	// sort bodies leaf to root
 	Bodies.Sort();
-	Algo::Reverse(Bodies);
 
 	// store pointers to bodies on bones (after sort!)
 	for (FRigidBody& Body : Bodies)
diff --git a/Engine/Plugins/Experimental/FullBodyIK/Source/PBIK/Public/Core/PBIKBody.h b/Engine/Plugins/Experimental/FullBodyIK/Source/PBIK/Public/Core/PBIKBody.h
--- a/Engine/Plugins/Experimental/FullBodyIK/Source/PBIK/Public/Core/PBIKBody.h
+++ b/Engine/Plugins/Experimental/FullBodyIK/Source/PBIK/Public/Core/PBIKBody.h
@@ -114,7 +114,9 @@ public:
 
 	FRigidBody* GetParentBody() const;
 
-	bool IsFullyLocked() const;
+	bool IsPositionLocked() const;
+
+	bool IsRotationLocked() const;
 
 	void ApplyPushToRotateBody(const FVector& Push, const FVector& Offset);
 	
@@ -126,7 +128,7 @@ public:
 // for sorting Bodies hierarchically (root to leaf order)
 inline bool operator<(const FRigidBody& Lhs, const FRigidBody& Rhs)
 { 
-	return Lhs.GetNumBonesToRoot() < Rhs.GetNumBonesToRoot(); 
+	return Lhs.GetNumBonesToRoot() > Rhs.GetNumBonesToRoot();
 }
 
 } // namespace
diff --git a/Engine/Plugins/Experimental/FullBodyIK/Source/PBIK/Public/Core/PBIKConstraint.h b/Engine/Plugins/Experimental/FullBodyIK/Source/PBIK/Public/Core/PBIKConstraint.h
--- a/Engine/Plugins/Experimental/FullBodyIK/Source/PBIK/Public/Core/PBIKConstraint.h
+++ b/Engine/Plugins/Experimental/FullBodyIK/Source/PBIK/Public/Core/PBIKConstraint.h
@@ -4,11 +4,11 @@
 
 #include "Math/Quat.h"
 
-namespace PBIK { struct FRigidBody; }
 struct FPBIKSolverSettings;
 
 namespace PBIK
 {
+	struct FRigidBody;
 
 struct FConstraint
 {
@@ -16,6 +16,7 @@ struct FConstraint
 
 	virtual void Solve(const FPBIKSolverSettings& Settings) = 0;
 	virtual void RemoveStretch(const float Percent){};
+	virtual void FinalPass(){};
 	virtual void UpdateFromInputs(){};
 };
 
@@ -49,6 +50,8 @@ private:
 	float AngleY;
 	float AngleZ;
 
+	bool bInFinalPass = false;
+
 public:
 
 	FJointConstraint(FRigidBody* InA, FRigidBody* InB);
@@ -59,6 +62,8 @@ public:
 
 	virtual void RemoveStretch(const float Percent) override;
 
+	virtual void FinalPass() override;
+
 	virtual void UpdateFromInputs() override;
 
 private:
diff --git a/Engine/Plugins/Experimental/FullBodyIK/Source/PBIK/Public/Core/PBIKSolver.h b/Engine/Plugins/Experimental/FullBodyIK/Source/PBIK/Public/Core/PBIKSolver.h
--- a/Engine/Plugins/Experimental/FullBodyIK/Source/PBIK/Public/Core/PBIKSolver.h
+++ b/Engine/Plugins/Experimental/FullBodyIK/Source/PBIK/Public/Core/PBIKSolver.h
@@ -96,8 +96,9 @@ struct FEffector
 		const FEffectorSettings& InSettings);
 
 	void UpdateChainStates();
-	void UpdateFromInputs(const FBone& SolverRoot);
+	void UpdateFromInputs(const FBone& InSolverRoot, const bool bAllowStretch);
 	float CalculateDistanceToChainRoot() const;
+	float CalculateDistanceToSolveRootAlongBones() const;
 	void ApplyPreferredAngles() const;
 };
 	
@@ -221,7 +222,7 @@ public:
 
 	PBIK_API bool Initialize();
 
-	PBIK_API void Solve(const FPBIKSolverSettings& Settings);
+	PBIK_API void Solve(const FPBIKSolverSettings& InSolverSettings);
 
 	PBIK_API void Reset();
 
@@ -276,7 +277,7 @@ private:
 
 	void UpdateBonesFromBodies();
 
-	void SolveConstraints(const FPBIKSolverSettings& Settings);
+	void SolveConstraints(const FPBIKSolverSettings& InSolverSettings);
 	
 	void ApplyRootPrePull(const EPBIKRootBehavior RootBehavior, const FRootPrePullSettings& PrePullSettings);
 	
