[FBIK] Fixed regression, prev change broke sub-solves and retarget engine tests

#virtualized

diff --git a/Engine/Plugins/Experimental/FullBodyIK/Source/PBIK/Private/Core/PBIKConstraint.cpp b/Engine/Plugins/Experimental/FullBodyIK/Source/PBIK/Private/Core/PBIKConstraint.cpp
--- a/Engine/Plugins/Experimental/FullBodyIK/Source/PBIK/Private/Core/PBIKConstraint.cpp
+++ b/Engine/Plugins/Experimental/FullBodyIK/Source/PBIK/Private/Core/PBIKConstraint.cpp
@@ -82,6 +82,11 @@ void FJointConstraint::RemoveStretch(const float Percent)
 
 void FJointConstraint::FinalPass()
 {
+	if (B->bIsLockedBySubSolve)
+	{
+		return;
+	}
+	
 	bInFinalPass = true;
 	UpdateJointLimits();
 	bInFinalPass = false;
diff --git a/Engine/Plugins/Experimental/FullBodyIK/Source/PBIK/Private/Core/PBIKSolver.cpp b/Engine/Plugins/Experimental/FullBodyIK/Source/PBIK/Private/Core/PBIKSolver.cpp
--- a/Engine/Plugins/Experimental/FullBodyIK/Source/PBIK/Private/Core/PBIKSolver.cpp
+++ b/Engine/Plugins/Experimental/FullBodyIK/Source/PBIK/Private/Core/PBIKSolver.cpp
@@ -513,7 +513,32 @@ void FPBIKSolver::SolveConstraints(const FPBIKSolverSettings& InSolverSettings)
 
 	auto RunConstraintPass = [this, &InSolverSettings](int32 CurrentIteration, int32 MaxIterations)
 	{
+		// solve all constraints
+		for (const TSharedPtr<PBIK::FConstraint>& Constraint : Constraints)
+		{
+			Constraint->Solve(InSolverSettings);
+		}
+
+		// on final iteration, run FinalPass
+		if (CurrentIteration == MaxIterations - 1)
+		{
+			// iterate in root-to-leaf order to force exact rotation limits
+			for (int32 C = Constraints.Num() - 1; C >= 0; --C)
+			{
+				Constraints[C]->FinalPass();
+			}
+		}
 		
+		// do post-pass to remove stretch
+		if (!InSolverSettings.bAllowStretch)
+		{
+			const float StretchAmount = FMath::Min(1.0f, static_cast<float>(CurrentIteration+1) / (MaxIterations / 2.0f));
+			// iterate in reverse order (root to leaf)
+			for (int32 C = Constraints.Num() - 1; C >= 0; --C)
+			{
+				Constraints[C]->RemoveStretch(StretchAmount);
+			}
+		}
 	};
 
 	// run pre-pass where we lock everything not in a sub-chain
@@ -539,34 +564,9 @@ void FPBIKSolver::SolveConstraints(const FPBIKSolverSettings& InSolverSettings)
 	}
 
 	// run main constraint pass
-	for (int32  CurrentIteration = 0; CurrentIteration < InSolverSettings.Iterations; ++CurrentIteration)
+	for (int32  I = 0; I < InSolverSettings.Iterations; ++I)
 	{
-		// solve all constraints
-		for (const TSharedPtr<PBIK::FConstraint>& Constraint : Constraints)
-		{
-			Constraint->Solve(InSolverSettings);
-		}
-
-		// on final iteration, run FinalPass
-		if (CurrentIteration == InSolverSettings.Iterations - 1)
-		{
-			// iterate in root-to-leaf order to force exact rotation limits
-			for (int32 C = Constraints.Num() - 1; C >= 0; --C)
-			{
-				Constraints[C]->FinalPass();
-			}
-		}
-		
-		// do post-pass to remove stretch
-		if (!InSolverSettings.bAllowStretch)
-		{
-			const float StretchAmount = FMath::Min(1.0f, static_cast<float>(CurrentIteration+1) / (InSolverSettings.Iterations / 2.0f));
-			// iterate in reverse order (root to leaf)
-			for (int32 C = Constraints.Num() - 1; C >= 0; --C)
-			{
-				Constraints[C]->RemoveStretch(StretchAmount);
-			}
-		}
+		RunConstraintPass(I, InSolverSettings.Iterations);
 	}
 }
 
@@ -746,6 +746,7 @@ bool FPBIKSolver::InitBodies()
 
 	// sort bodies leaf to root
 	Bodies.Sort();
+	Algo::Reverse(Bodies);
 
 	// store pointers to bodies on bones (after sort!)
 	for (FRigidBody& Body : Bodies)
diff --git a/Engine/Plugins/Experimental/FullBodyIK/Source/PBIK/Public/Core/PBIKBody.h b/Engine/Plugins/Experimental/FullBodyIK/Source/PBIK/Public/Core/PBIKBody.h
--- a/Engine/Plugins/Experimental/FullBodyIK/Source/PBIK/Public/Core/PBIKBody.h
+++ b/Engine/Plugins/Experimental/FullBodyIK/Source/PBIK/Public/Core/PBIKBody.h
@@ -128,7 +128,7 @@ public:
 // for sorting Bodies hierarchically (root to leaf order)
 inline bool operator<(const FRigidBody& Lhs, const FRigidBody& Rhs)
 { 
-	return Lhs.GetNumBonesToRoot() > Rhs.GetNumBonesToRoot();
+	return Lhs.GetNumBonesToRoot() < Rhs.GetNumBonesToRoot();
 }
 
 } // namespace
