[FBIK] Fixed bug, excluding parent of effector causes effector not to converge correctly.

diff --git a/Engine/Plugins/Experimental/FullBodyIK/Source/PBIK/Private/Core/PBIKBody.cpp b/Engine/Plugins/Experimental/FullBodyIK/Source/PBIK/Private/Core/PBIKBody.cpp
--- a/Engine/Plugins/Experimental/FullBodyIK/Source/PBIK/Private/Core/PBIKBody.cpp
+++ b/Engine/Plugins/Experimental/FullBodyIK/Source/PBIK/Private/Core/PBIKBody.cpp
@@ -56,9 +56,6 @@ void FRigidBody::Initialize(const FBone* SolverRoot)
 {
 	// calculate transform and mass of body based on the skeleton
 	UpdateTransformAndMassFromBones();
-
-	// store initial rotation from ref pose
-	InitialRotation = Rotation;
 	
 	// calculate num bones distance to root
 	NumBonesToRoot = 0;
diff --git a/Engine/Plugins/Experimental/FullBodyIK/Source/PBIK/Private/Core/PBIKConstraint.cpp b/Engine/Plugins/Experimental/FullBodyIK/Source/PBIK/Private/Core/PBIKConstraint.cpp
--- a/Engine/Plugins/Experimental/FullBodyIK/Source/PBIK/Private/Core/PBIKConstraint.cpp
+++ b/Engine/Plugins/Experimental/FullBodyIK/Source/PBIK/Private/Core/PBIKConstraint.cpp
@@ -302,17 +302,14 @@ float FJointConstraint::SignedAngleBetweenNormals(
 	return Dot >= 0 ? Angle : -Angle;
 }
 
-FPinConstraint::FPinConstraint(
-	FRigidBody* InBody,
-	const FVector& InPinPositionOrig,
-	const FQuat& InPinRotationOrig)
+FPinConstraint::FPinConstraint(FRigidBody* InBody, FBone* InPinBone)
 {
-	GoalPosition = InPinPositionOrig;
-	GoalRotation = InPinRotationOrig;
+	PinBone = InPinBone;
+	Body = InBody;
+	PinPointLocalToA = Body->Rotation.Inverse() * (PinBone->Position - Body->Position);
 	
-	A = InBody;
-	PinPointLocalToA = A->Rotation.Inverse() * (GoalPosition - A->Position);
-	ARotLocalToPin = A->Rotation * GoalRotation.Inverse();
+	GoalPosition = InPinBone->Position;
+	GoalRotation = InPinBone->Rotation;
 }
 
 void FPinConstraint::Solve(const FPBIKSolverSettings& Settings)
@@ -327,13 +324,13 @@ void FPinConstraint::Solve(const FPBIKSolverSettings& Settings)
 	FVector Correction = GetPositionCorrection(AToPinPoint);
 		
 	// rotate body from alignment of pin points
-	A->ApplyPushToRotateBody(Correction, AToPinPoint);
+	Body->ApplyPushToRotateBody(Correction, AToPinPoint);
 
 	// apply positional correction to Body to align with target (after rotation)
 	// (applying directly without considering PositionStiffness because PinConstraints need
 	// to precisely pull the attached body to achieve convergence)
 	Correction = GetPositionCorrection(AToPinPoint);
-	A->Position += Correction * Settings.OverRelaxation; 
+	Body->Position += Correction * Settings.OverRelaxation; 
 }
 
 void FPinConstraint::SetGoal(const FVector& InGoalPosition, const FQuat& InGoalRotation, const float InAlpha)
@@ -343,10 +340,15 @@ void FPinConstraint::SetGoal(const FVector& InGoalPosition, const FQuat& InGoalR
 	Alpha = InAlpha;
 }
 
+void FPinConstraint::UpdateFromInputs()
+{
+	PinPointLocalToA = Body->Rotation.Inverse() * (PinBone->Position - Body->Position);
+}
+
 FVector FPinConstraint::GetPositionCorrection(FVector& OutBodyToPinPoint) const
 {
-	OutBodyToPinPoint = A->Rotation * PinPointLocalToA;
-	const FVector PinPoint = A->Position + OutBodyToPinPoint;
+	OutBodyToPinPoint = Body->Rotation * PinPointLocalToA;
+	const FVector PinPoint = Body->Position + OutBodyToPinPoint;
 	return (GoalPosition - PinPoint) * Alpha;
 }
 
diff --git a/Engine/Plugins/Experimental/FullBodyIK/Source/PBIK/Private/Core/PBIKSolver.cpp b/Engine/Plugins/Experimental/FullBodyIK/Source/PBIK/Private/Core/PBIKSolver.cpp
--- a/Engine/Plugins/Experimental/FullBodyIK/Source/PBIK/Private/Core/PBIKSolver.cpp
+++ b/Engine/Plugins/Experimental/FullBodyIK/Source/PBIK/Private/Core/PBIKSolver.cpp
@@ -14,6 +14,7 @@ namespace PBIK
 	{
 		check(InBone);
 		Bone = InBone;
+		Bone->Effector = this;
 		SetGoal(Bone->Position, Bone->Rotation, Settings);
 	}
 
@@ -240,12 +241,6 @@ void FPBIKSolver::Solve(const FPBIKSolverSettings& InSolverSettings)
 		Body.UpdateFromInputs(InSolverSettings);
 	}
 
-	// update constraints with new local pin locations relative to (possibly) translated joint positions
-	for (const TSharedPtr<PBIK::FConstraint>& Constraint : Constraints)
-	{
-		Constraint->UpdateFromInputs();
-	}
-
 	// optionally pin root in-place (convenience, does not require an effector)
 	if (InSolverSettings.RootBehavior == EPBIKRootBehavior::PinToInput)
 	{
@@ -264,6 +259,12 @@ void FPBIKSolver::Solve(const FPBIKSolverSettings& InSolverSettings)
 		Effector.UpdateFromInputs(*SolverRoot, InSolverSettings.bAllowStretch);
 	}
 
+	// update constraints with new local pin locations relative to (possibly) translated joint positions
+	for (const TSharedPtr<PBIK::FConstraint>& Constraint : Constraints)
+	{
+		Constraint->UpdateFromInputs();
+	}
+
 	// do all constraint solving
 	UpdateBodies(InSolverSettings);
 
@@ -294,55 +295,49 @@ void FPBIKSolver::UpdateBonesFromBodies()
 	using PBIK::FRigidBody;
 	using PBIK::FBone;
 	using PBIK::FEffector;
-	
-	// update Bone transforms controlled by Bodies
-	for (FRigidBody& Body : Bodies)
-	{
-		Body.Bone->Position = Body.Position + Body.Rotation * Body.BoneLocalPosition;
-		Body.Bone->Rotation = Body.Rotation;
-	}
 
-	// update Bone transforms controlled by effectors
-	for (const FEffector& Effector : Effectors)
+
+	// iterate over ALL bones in root-to-leaf order
+	for (FBone& Bone : Bones)
 	{
-		FBone* Bone = Effector.Bone;
-		if (Bone->bIsSolverRoot)
+		// if bone is controlled by a body, then update it from the body local offset
+		if (Bone.Body)
 		{
-			continue; // if there's an effector on the root, leave it where the body ended up
+			Bone.Position = Bone.Body->Position + Bone.Body->Rotation * Bone.Body->BoneLocalPosition;
+			Bone.Rotation = Bone.Body->Rotation;
 		}
-
-		// if effector is between other effectors, it will have an associated rigid body
-		// so leave transform where body ended up after solve (applied above)
-		// otherwise, snap position back to FK location
-		if (!Bone->Body)
+		
+		// if bone is controlled by an effector, we optionally apply the effector rotation
+		if (Bone.Effector)
 		{
-			if (Bone->Parent)
-			{
-				Bone->Position = Bone->Parent->Position + Bone->Parent->Rotation * Bone->LocalPositionFromInput;
-			}
-			else
+			// if effector is between other effectors, it will have an associated rigid body
+			// so leave transform where body ended up after solve (applied above)
+			// otherwise, snap position back to FK location
+			if (!Bone.Body)
 			{
-				// effector was on the solver root
-				Bone->Position = Effector.Position;
+				if (Bone.Parent)
+				{
+					Bone.Position = Bone.Parent->Position + Bone.Parent->Rotation * Bone.LocalPositionFromInput;
+				}
+				else
+				{
+					// effector was on the solver root
+					Bone.Position = Bone.Effector->Position;
+				}
 			}
-		}
 
-		// optionally pin rotation to that of effector
-		const FQuat RotationWithoutEffector = Bone->Parent ? Bone->Parent->Rotation * Bone->LocalRotationFromInput : Bone->Rotation; 
-		const float RotAmount = FMath::Clamp(Effector.Settings.PinRotation, 0.0f, 1.0f);
-		Bone->Rotation = FQuat::FastLerp(RotationWithoutEffector, Effector.Rotation, RotAmount).GetNormalized();
-	}
+			// optionally pin rotation to that of effector
+			const FQuat RotationWithoutEffector = Bone.Parent ? Bone.Parent->Rotation * Bone.LocalRotationFromInput : Bone.Rotation; 
+			const float RotAmount = FMath::Clamp(Bone.Effector->Settings.PinRotation, 0.0f, 1.0f);
+			Bone.Rotation = FQuat::FastLerp(RotationWithoutEffector, Bone.Effector->Rotation, RotAmount).GetNormalized();
+		}
 
-	// propagate to non-solved bones (requires storage in root to tip order)
-	for (FBone& Bone : Bones)
-	{
-		if (Bone.bIsSolved || !Bone.Parent)
+		// trivial leaf bones (children of simulated bones)
+		if (!Bone.bIsSolved && Bone.Parent)
 		{
-			continue;
+			Bone.Position = Bone.Parent->Position + Bone.Parent->Rotation * Bone.LocalPositionFromInput;
+			Bone.Rotation = Bone.Parent->Rotation * Bone.LocalRotationFromInput;
 		}
-
-		Bone.Position = Bone.Parent->Position + Bone.Parent->Rotation * Bone.LocalPositionFromInput;
-		Bone.Rotation = Bone.Parent->Rotation * Bone.LocalRotationFromInput;
 	}
 }
 
@@ -792,7 +787,7 @@ bool FPBIKSolver::InitConstraints()
 		}
 
 		FRigidBody* Body = BodyBone->Body;
-		TSharedPtr<FPinConstraint> Constraint = MakeShared<FPinConstraint>(Body, Effector.Position, Effector.Rotation);
+		TSharedPtr<FPinConstraint> Constraint = MakeShared<FPinConstraint>(Body, Effector.Bone);
 		Effector.Pin = Constraint;
 		Body->Effector = &Effector;
 		Body->Pin = Constraint.Get();
diff --git a/Engine/Plugins/Experimental/FullBodyIK/Source/PBIK/Public/Core/PBIKBody.h b/Engine/Plugins/Experimental/FullBodyIK/Source/PBIK/Public/Core/PBIKBody.h
--- a/Engine/Plugins/Experimental/FullBodyIK/Source/PBIK/Public/Core/PBIKBody.h
+++ b/Engine/Plugins/Experimental/FullBodyIK/Source/PBIK/Public/Core/PBIKBody.h
@@ -26,6 +26,7 @@ struct FBone
 
 	// initialized - these fields are null/empty until after Solver::Initialize()
 	FRigidBody* Body = nullptr;
+	const FEffector* Effector = nullptr;
 	FBone* Parent = nullptr;
 	TArray<FBone*> Children;
 	bool bIsSolverRoot = false;
diff --git a/Engine/Plugins/Experimental/FullBodyIK/Source/PBIK/Public/Core/PBIKConstraint.h b/Engine/Plugins/Experimental/FullBodyIK/Source/PBIK/Public/Core/PBIKConstraint.h
--- a/Engine/Plugins/Experimental/FullBodyIK/Source/PBIK/Public/Core/PBIKConstraint.h
+++ b/Engine/Plugins/Experimental/FullBodyIK/Source/PBIK/Public/Core/PBIKConstraint.h
@@ -9,6 +9,7 @@ struct FPBIKSolverSettings;
 namespace PBIK
 {
 	struct FRigidBody;
+	struct FBone;
 
 struct FConstraint
 {
@@ -100,23 +101,22 @@ private:
 	
 	FVector GoalPosition;
 	FQuat GoalRotation;
+	float Alpha = 1.0;
 	
-	FRigidBody* A;
+	FRigidBody* Body;
+	FBone* PinBone;
 	FVector PinPointLocalToA;
-	float Alpha = 1.0;
-	FQuat ARotLocalToPin;
 
 public:
 	
-	FPinConstraint(
-		FRigidBody* InBody,
-		const FVector& InPinPositionOrig,
-		const FQuat& InPinRotationOrig);
+	FPinConstraint(FRigidBody* InBody, FBone* InPinBone);
 
 	virtual ~FPinConstraint() {};
 
 	virtual void Solve(const FPBIKSolverSettings& Settings) override;
 
+	virtual void UpdateFromInputs() override;
+
 	void SetGoal(const FVector& InGoalPosition, const FQuat& InGoalRotation, const float InAlpha);
 
 private:
