[FullBodyIK] Fixed regression causing NANs on effector.

diff --git a/Engine/Plugins/Experimental/FullBodyIK/Source/PBIK/Private/Core/PBIKSolver.cpp b/Engine/Plugins/Experimental/FullBodyIK/Source/PBIK/Private/Core/PBIKSolver.cpp
--- a/Engine/Plugins/Experimental/FullBodyIK/Source/PBIK/Private/Core/PBIKSolver.cpp
+++ b/Engine/Plugins/Experimental/FullBodyIK/Source/PBIK/Private/Core/PBIKSolver.cpp
@@ -14,7 +14,6 @@ namespace PBIK
 	{
 		check(InBone);
 		Bone = InBone;
-		Bone->Effector = this;
 		SetGoal(Bone->Position, Bone->Rotation, Settings);
 	}
 
@@ -23,13 +22,9 @@ namespace PBIK
 		const FQuat& InRotationGoal,
 		const FEffectorSettings& InSettings)
 	{
-		PositionOrig = Bone->Position;
-		RotationOrig = Bone->Rotation;
-
-		Position = PositionGoal = InPositionGoal;
-		Rotation = RotationGoal = InRotationGoal;
-
 		Settings = InSettings;
+		PositionGoal = InPositionGoal;
+		RotationGoal = InRotationGoal;
 	}
 
 	void FEffector::UpdateChainStates()
@@ -67,20 +62,24 @@ namespace PBIK
 
 	void FEffector::UpdateFromInputs(const FBone& InSolverRoot, const bool bAllowStretch)
 	{
-		// blend effector transform by alpha
-		Position = FMath::Lerp(PositionOrig, PositionGoal, Settings.PositionAlpha);
-		Rotation = FMath::Lerp(RotationOrig, RotationGoal, Settings.RotationAlpha);
+		// update original transform based on input bone pose
+		PositionOrig = Bone->Position;
+		RotationOrig = Bone->Rotation;
+		
+		// update final position/rotation by blending from input bone pose to goal by alpha
+		PositionFinal = FMath::Lerp(PositionOrig, PositionGoal, Settings.PositionAlpha);
+		RotationFinal = FQuat::Slerp(RotationOrig, RotationGoal, Settings.RotationAlpha);
 
 		// if the solver root is pinned, we clamp the effector position so that it cannot pull excessively,
 		// this prevents extreme constraint forces that cause solve instability
 		if (FMath::IsNearlyZero(InSolverRoot.Body->InvMass) && !bAllowStretch)
 		{
 			const float DistToSolveRootAlongBones = CalculateDistanceToSolveRootAlongBones();
-			Position = InSolverRoot.Position + (Position - InSolverRoot.Position).GetClampedToMaxSize(DistToSolveRootAlongBones);
+			PositionFinal = InSolverRoot.Position + (PositionFinal - InSolverRoot.Position).GetClampedToMaxSize(DistToSolveRootAlongBones);
 		}
 
-		// set pin goal transform
-		Pin.Pin()->SetGoal(Position, Rotation, Settings.StrengthAlpha);
+		// set pin constraint final transform
+		Pin.Pin()->SetGoal(PositionFinal, RotationFinal, Settings.StrengthAlpha);
 		
 		// update length of chain this effector controls in the input pose
 		DistToChainRootInInputPose = CalculateDistanceToChainRoot();
@@ -147,8 +146,8 @@ namespace PBIK
 		// if the parent sub-root is attached to an effector, use the effector's position
 		// otherwise use the current position of the FRigidBody
 		const FEffector* ParentEffector = ChainRootBody->Effector;
-		const FVector ParentSubRootPosition = ParentEffector ? ParentEffector->Position : ChainRootBody->Position;
-		const float DistToSubRootCurrent = (ParentSubRootPosition - Position).Size();
+		const FVector ParentSubRootPosition = ParentEffector ? ParentEffector->PositionFinal : ChainRootBody->Position;
+		const float DistToSubRootCurrent = (ParentSubRootPosition - PositionFinal).Size();
 		if (DistToSubRootCurrent >= DistToChainRootInInputPose)
 		{
 			return; // limb is stretched
@@ -308,8 +307,11 @@ void FPBIKSolver::UpdateBonesFromBodies()
 		}
 		
 		// if bone is controlled by an effector, we optionally apply the effector rotation
-		if (Bone.Effector)
+		if (Bone.EffectorIndex != INDEX_NONE)
 		{
+			// get the effector associated with this bone
+			const FEffector& Effector = Effectors[Bone.EffectorIndex];
+			
 			// if effector is between other effectors, it will have an associated rigid body
 			// so leave transform where body ended up after solve (applied above)
 			// otherwise, snap position back to FK location
@@ -322,14 +324,14 @@ void FPBIKSolver::UpdateBonesFromBodies()
 				else
 				{
 					// effector was on the solver root
-					Bone.Position = Bone.Effector->Position;
+					Bone.Position = Effector.PositionFinal;
 				}
 			}
 
 			// optionally pin rotation to that of effector
 			const FQuat RotationWithoutEffector = Bone.Parent ? Bone.Parent->Rotation * Bone.LocalRotationFromInput : Bone.Rotation; 
-			const float RotAmount = FMath::Clamp(Bone.Effector->Settings.PinRotation, 0.0f, 1.0f);
-			Bone.Rotation = FQuat::FastLerp(RotationWithoutEffector, Bone.Effector->Rotation, RotAmount).GetNormalized();
+			const float RotAmount = FMath::Clamp(Effector.Settings.PinRotation, 0.0f, 1.0f);
+			Bone.Rotation = FQuat::Slerp(RotationWithoutEffector, Effector.RotationFinal, RotAmount);
 		}
 
 		// trivial leaf bones (children of simulated bones)
@@ -359,7 +361,7 @@ void FPBIKSolver::ApplyRootPrePull(const EPBIKRootBehavior RootBehavior, const F
 	for (FEffector& Effector : Effectors)
 	{
 		InitialPoints.Add(Effector.PositionOrig);
-		CurrentPoints.Add(Effector.Position);
+		CurrentPoints.Add(Effector.PositionFinal);
 	}
 
 	// calculate "best fit" rotational delta from deformed effector configuration
@@ -439,20 +441,20 @@ void FPBIKSolver::ApplyPullChainAlpha(const float GlobalPullChainAlpha)
 		{
 			// add effector offset to chain root
 			const FEffector* RootEffector = Effector.ChainRootBody->Effector;
-			FVector RootEffectorDelta = RootEffector->Position - RootEffector->PositionOrig;
+			FVector RootEffectorDelta = RootEffector->PositionFinal - RootEffector->PositionOrig;
 			ChainStartNew = Effector.ChainRootBody->Position + RootEffectorDelta;
 		}else
 		{
 			ChainStartNew = Effector.ChainRootBody->Position;
 		}
-		const FVector ChainEndNew = Effector.Position;
+		const FVector ChainEndNew = Effector.PositionFinal;
 		FVector ChainVecNew;
 		float ChainLenNew;
 		(ChainEndNew - ChainStartNew).ToDirectionAndLength(ChainVecNew, ChainLenNew);
 
 		FQuat ChainDeltaRotation = FQuat::FindBetweenNormals(ChainVecOrig, ChainVecNew);
 		const float ChainDeltaAlpha = Effector.Settings.PullChainAlpha * Effector.Settings.StrengthAlpha * GlobalPullChainAlpha;
-		ChainDeltaRotation = FQuat::FastLerp(FQuat::Identity, ChainDeltaRotation, ChainDeltaAlpha);
+		ChainDeltaRotation = FQuat::FastLerp(FQuat::Identity, ChainDeltaRotation, ChainDeltaAlpha).GetNormalized();
 		FVector ChainDeltaPosition = ChainVecNew * (ChainLenNew - ChainLenOrig) * ChainDeltaAlpha;
 
 		const FBone* Bone = Effector.Bone->Parent;
@@ -865,10 +867,10 @@ void FPBIKSolver::Reset()
 	
 	bReadyToSimulate = false;
 	SolverRoot = nullptr;
-	Bodies.Empty();
-	Bones.Empty();
-	Constraints.Empty();
-	Effectors.Empty();
+	Bodies.Reset();
+	Bones.Reset();
+	Constraints.Reset();
+	Effectors.Reset();
 }
 
 bool FPBIKSolver::IsReadyToSimulate() const
@@ -892,8 +894,9 @@ int32 FPBIKSolver::AddEffector(const FName BoneName)
 	{
 		if (Bone.Name == BoneName)
 		{
-			Effectors.Emplace(&Bone);
-			return Effectors.Num() - 1;
+			const int32 NewEffectorIndex = Effectors.Emplace(&Bone);
+			Bone.EffectorIndex = NewEffectorIndex;
+			return NewEffectorIndex;
 		}
 	}
 
@@ -950,12 +953,12 @@ int32 FPBIKSolver::GetBoneIndex(FName BoneName) const
 
 void FPBIKSolver::SetEffectorGoal(
 	const int32 Index,
-	const FVector& InPosition,
-	const FQuat& InRotation,
+	const FVector& InGoalPosition,
+	const FQuat& InGoalRotation,
 	const PBIK::FEffectorSettings& InSettings)
 {
 	check(Index >= 0 && Index < Effectors.Num());
-	Effectors[Index].SetGoal(InPosition, InRotation, InSettings);
+	Effectors[Index].SetGoal(InGoalPosition, InGoalRotation, InSettings);
 }
 
 // given a set of points in an initial configuration and the same set of points in a deformed configuration,
diff --git a/Engine/Plugins/Experimental/FullBodyIK/Source/PBIK/Public/Core/PBIKBody.h b/Engine/Plugins/Experimental/FullBodyIK/Source/PBIK/Public/Core/PBIKBody.h
--- a/Engine/Plugins/Experimental/FullBodyIK/Source/PBIK/Public/Core/PBIKBody.h
+++ b/Engine/Plugins/Experimental/FullBodyIK/Source/PBIK/Public/Core/PBIKBody.h
@@ -26,9 +26,9 @@ struct FBone
 
 	// initialized - these fields are null/empty until after Solver::Initialize()
 	FRigidBody* Body = nullptr;
-	const FEffector* Effector = nullptr;
 	FBone* Parent = nullptr;
 	TArray<FBone*> Children;
+	int32 EffectorIndex = INDEX_NONE;
 	bool bIsSolverRoot = false;
 	bool bIsSolved = false;
 	bool bIsSubRoot = false;
diff --git a/Engine/Plugins/Experimental/FullBodyIK/Source/PBIK/Public/Core/PBIKSolver.h b/Engine/Plugins/Experimental/FullBodyIK/Source/PBIK/Public/Core/PBIKSolver.h
--- a/Engine/Plugins/Experimental/FullBodyIK/Source/PBIK/Public/Core/PBIKSolver.h
+++ b/Engine/Plugins/Experimental/FullBodyIK/Source/PBIK/Public/Core/PBIKSolver.h
@@ -62,8 +62,8 @@ struct FEffectorSettings
 
 struct FEffector
 {
-	FVector Position;
-	FQuat Rotation;
+	FVector PositionFinal;
+	FQuat RotationFinal;
 
 	FVector PositionOrig;
 	FQuat RotationOrig;
@@ -232,8 +232,8 @@ public:
 
 	PBIK_API void SetEffectorGoal(
 		const int32 Index, 
-		const FVector& InPosition, 
-		const FQuat& InRotation, 
+		const FVector& InGoalPosition, 
+		const FQuat& InGoalRotation, 
 		const PBIK::FEffectorSettings& Settings);
 
 	PBIK_API void GetBoneGlobalTransform(const int32 Index, FTransform& OutTransform);
