[FBIK] Profiled and optimized FBIK solver.
- Joint constraint solve simplified
- Pull Chain Alpha more accurate
- Unsolved bones culled from update
- Minor optimizations throughout

diff --git a/Engine/Plugins/Experimental/FullBodyIK/Source/PBIK/Private/Core/PBIKBody.cpp b/Engine/Plugins/Experimental/FullBodyIK/Source/PBIK/Private/Core/PBIKBody.cpp
--- a/Engine/Plugins/Experimental/FullBodyIK/Source/PBIK/Private/Core/PBIKBody.cpp
+++ b/Engine/Plugins/Experimental/FullBodyIK/Source/PBIK/Private/Core/PBIKBody.cpp
@@ -36,13 +36,14 @@ bool FBone::HasChild(const FBone* Bone)
 
 void FBone::UpdateFromInputs()
 {
-	if (!Parent)
+	if (!bIsPartOfSolvedBranch || !Parent)
 	{
 		return;
 	}
 
-	LocalPositionFromInput = Parent->Rotation.Inverse() * (Position - Parent->Position);
-	LocalRotationFromInput = Parent->Rotation.Inverse() * Rotation;
+	const FQuat ParentRotInv = Parent->Rotation.Inverse();
+	LocalPositionFromInput = ParentRotInv * (Position - Parent->Position);
+	LocalRotationFromInput = ParentRotInv * Rotation;
 	Length = LocalPositionFromInput.Size();
 }
 
@@ -92,13 +93,6 @@ void FRigidBody::UpdateTransformAndMassFromBones()
 	Position = InputPosition = Centroid;
 	Rotation = InitialRotation = Bone->Rotation;
 	BoneLocalPosition = Bone->Rotation.Inverse() * (Bone->Position - Centroid);
-
-	ChildLocalPositions.Empty(Bone->Children.Num());
-	for (const FBone* Child : Bone->Children)
-	{
-		FVector ChildLocalPos = Rotation.Inverse() * (Child->Position - Centroid);
-		ChildLocalPositions.Add(ChildLocalPos);
-	}
 }
 
 int FRigidBody::GetNumBonesToRoot() const
@@ -164,17 +158,7 @@ void FRigidBody::ApplyRotationDelta(const FQuat& DeltaQ)
 	{
 		ClampedDQ *= MaxPhi / Phi;
 	}
-
-	/** DeltaQ is assumed to be a "pure" quaternion representing an infintesimal rotation */
-	FQuat Delta = ClampedDQ * Rotation;
-	Delta.X *= 0.5f;
-	Delta.Y *= 0.5f;
-	Delta.Z *= 0.5f;
-	Delta.W *= 0.5f;
-	Rotation.X = Rotation.X + Delta.X;
-	Rotation.Y = Rotation.Y + Delta.Y;
-	Rotation.Z = Rotation.Z + Delta.Z;
-	Rotation.W = Rotation.W + Delta.W;
-	Rotation.Normalize();
+	
+	Rotation += (ClampedDQ * Rotation) * 0.5f;
 }
 } // namespace
diff --git a/Engine/Plugins/Experimental/FullBodyIK/Source/PBIK/Private/Core/PBIKConstraint.cpp b/Engine/Plugins/Experimental/FullBodyIK/Source/PBIK/Private/Core/PBIKConstraint.cpp
--- a/Engine/Plugins/Experimental/FullBodyIK/Source/PBIK/Private/Core/PBIKConstraint.cpp
+++ b/Engine/Plugins/Experimental/FullBodyIK/Source/PBIK/Private/Core/PBIKConstraint.cpp
@@ -32,52 +32,42 @@ void FJointConstraint::Solve(const FPBIKSolverSettings& Settings)
 	{
 		return;
 	}
-	
-	// get pos correction to use for rotation
-	FVector OffsetA;
-	FVector OffsetB;
-	FVector Correction = GetPositionCorrection(OffsetA, OffsetB);
-
-	// apply rotation from correction 
-	// at pin point to both bodies
-	A->ApplyPushToRotateBody(Correction, OffsetA);
-	B->ApplyPushToRotateBody(-Correction, OffsetB);
-
-	// enforce joint limits
-	UpdateJointLimits();
 
 	// calc inv mass of body A and B
-	const float WA = A->InvMass;
-	const float WB = B->InvMass;
-	const float W = WA + WB;
+	const float W = A->InvMass + B->InvMass;
 	if (FMath::IsNearlyZero(W))
 	{
 		return; // no correction can be applied on full locked configuration
 	}
 
-	// apply positional correction to align pin point on both Bodies
-	// NOTE: applying position correction AFTER rotation takes into
-	// consideration change in relative pin locations mid-step after
-	// bodies have been rotated by ApplyPushToRotateBody
-	Correction = GetPositionCorrection(OffsetA, OffsetB);
+	// get pos correction to use for rotation
+	FVector OffsetA;
+	FVector OffsetB;
+	FVector Correction = GetPositionCorrection(OffsetA, OffsetB);
 	FVector N;
 	float C;
 	Correction.ToDirectionAndLength(N, C);
+	const FVector CorrectDir = N * (C * (1.0f / W));
+	const FVector CorrectA = CorrectDir * A->InvMass;
+	const FVector CorrectB = -CorrectDir * B->InvMass;
 
-	const FVector CorrectA = N * (C * (WA / W));
-	const FVector CorrectB = N * (C * (WB / W));
+	// apply rotation from correction at pin point to both bodies
+	A->ApplyPushToRotateBody(CorrectA, OffsetA);
+	B->ApplyPushToRotateBody(CorrectB, OffsetB);
 
+	// enforce joint limits
+	UpdateJointLimits();
+	
+	// apply position correction
 	A->ApplyPositionDelta(CorrectA);
-	B->ApplyPositionDelta(-CorrectB);
+	B->ApplyPositionDelta(CorrectB);
 }
 
 void FJointConstraint::RemoveStretch(const float Percent)
 {
-	// apply position correction, keeping parent body fixed
-	FVector ToA;
-	FVector ToB;
-	const FVector Correction = GetPositionCorrection(ToA, ToB);
-	B->Position -= Correction * Percent;
+	const FVector PinPointOnA = A->Position + A->Rotation * PinPointLocalToA;
+	const FVector PinPointOnB = B->Position + B->Rotation * PinPointLocalToB;
+	B->Position -= (PinPointOnB - PinPointOnA) * Percent;
 }
 
 void FJointConstraint::FinalPass()
@@ -117,7 +107,7 @@ void FJointConstraint::ApplyRotationCorrection(FQuat DeltaQ) const
 void FJointConstraint::UpdateJointLimits()
 {
 	FBoneSettings& J = B->J;
-
+	
 	// no limits at all
 	if (J.X == ELimitType::Free &&
 		J.Y == ELimitType::Free &&
@@ -340,6 +330,11 @@ void FPinConstraint::SetGoal(const FVector& InGoalPosition, const FQuat& InGoalR
 	Alpha = InAlpha;
 }
 
+FVector FPinConstraint::GetCurrentPinPoint() const
+{
+	return Body->Position + (Body->Rotation * PinPointLocalToA);
+}
+
 void FPinConstraint::UpdateFromInputs()
 {
 	PinPointLocalToA = Body->Rotation.Inverse() * (PinBone->Position - Body->Position);
diff --git a/Engine/Plugins/Experimental/FullBodyIK/Source/PBIK/Private/Core/PBIKDebug.cpp b/Engine/Plugins/Experimental/FullBodyIK/Source/PBIK/Private/Core/PBIKDebug.cpp
--- a/Engine/Plugins/Experimental/FullBodyIK/Source/PBIK/Private/Core/PBIKDebug.cpp
+++ b/Engine/Plugins/Experimental/FullBodyIK/Source/PBIK/Private/Core/PBIKDebug.cpp
@@ -6,17 +6,4 @@
 
 namespace PBIK
 {
-	void FDebugDraw::GetDebugLinesForBodies(TArray<FDebugLine>& OutLines)
-	{
-		OutLines.Empty();
-		for (const FRigidBody& Body : Solver->Bodies)
-		{
-			if (Body.Bone->Children.Num() == 1)
-			{
-				FVector StartPos = Body.Position + Body.Rotation * Body.BoneLocalPosition;
-				FVector EndPos = Body.Position + Body.Rotation * Body.ChildLocalPositions[0];
-				OutLines.Emplace(StartPos, EndPos);
-			}
-		}
-	}
 }
diff --git a/Engine/Plugins/Experimental/FullBodyIK/Source/PBIK/Private/Core/PBIKSolver.cpp b/Engine/Plugins/Experimental/FullBodyIK/Source/PBIK/Private/Core/PBIKSolver.cpp
--- a/Engine/Plugins/Experimental/FullBodyIK/Source/PBIK/Private/Core/PBIKSolver.cpp
+++ b/Engine/Plugins/Experimental/FullBodyIK/Source/PBIK/Private/Core/PBIKSolver.cpp
@@ -34,6 +34,7 @@ namespace PBIK
 		ChainRootBody = nullptr;
 		const int32 MaxDepth = Settings.ChainDepth;
 		const FBone* Parent = Bone->Parent;
+		const FBone* Child = Bone;
 		int32 Depth = 1;
 		while (Parent)
 		{
@@ -48,12 +49,13 @@ namespace PBIK
 			const bool bIsAtChainRoot = MaxDepth > 0 ? Depth > MaxDepth || bIsAtAnyRoot : bIsAtAnyRoot;
 			if (bIsAtChainRoot)
 			{
-				ChainRootBody = Parent->Body;
+				ChainRootBody = Child->Body;
 				break;
 			}
 			
 			Parent->Body->bIsPartOfSubChain = true;
 			++Depth;
+			Child = Parent;
 			Parent = Parent->Parent;
 		}
 
@@ -79,7 +81,7 @@ namespace PBIK
 		}
 
 		// set pin constraint final transform
-		Pin.Pin()->SetGoal(PositionFinal, RotationFinal, Settings.StrengthAlpha);
+		Pin->SetGoal(PositionFinal, RotationFinal, Settings.StrengthAlpha);
 		
 		// update length of chain this effector controls in the input pose
 		DistToChainRootInInputPose = CalculateDistanceToChainRoot();
@@ -263,7 +265,7 @@ void FPBIKSolver::Solve(const FPBIKSolverSettings& InSolverSettings)
 	{
 		Constraint->UpdateFromInputs();
 	}
-
+		
 	// do all constraint solving
 	UpdateBodies(InSolverSettings);
 
@@ -274,7 +276,7 @@ void FPBIKSolver::Solve(const FPBIKSolverSettings& InSolverSettings)
 void FPBIKSolver::UpdateBodies(const FPBIKSolverSettings& Settings)
 {
 	using PBIK::FRigidBody;
-	
+
 	// apply large scale gross movement to entire skeleton based on effector deltas
 	// (this creates a better pose from which to start the constraint solving)
 	ApplyRootPrePull(Settings.RootBehavior, Settings.PrePullRootSettings);
@@ -284,7 +286,7 @@ void FPBIKSolver::UpdateBodies(const FPBIKSolverSettings& Settings)
 
 	// pull each effector's chain towards itself
 	ApplyPullChainAlpha(Settings.GlobalPullChainAlpha);
-		
+
 	// run ALL constraint iterations
 	SolveConstraints(Settings);
 }
@@ -295,10 +297,23 @@ void FPBIKSolver::UpdateBonesFromBodies()
 	using PBIK::FBone;
 	using PBIK::FEffector;
 
-
 	// iterate over ALL bones in root-to-leaf order
 	for (FBone& Bone : Bones)
 	{
+		// bone is located outside the solving branch in the hierarchy
+		if (!Bone.bIsPartOfSolvedBranch)
+		{
+			continue;
+		}
+
+		// trivial leaf bones (children of simulated bones)
+		if (!Bone.bIsSolved && Bone.Parent)
+		{
+			Bone.Position = Bone.Parent->Position + Bone.Parent->Rotation * Bone.LocalPositionFromInput;
+			Bone.Rotation = Bone.Parent->Rotation * Bone.LocalRotationFromInput;
+			continue;
+		}
+		
 		// if bone is controlled by a body, then update it from the body local offset
 		if (Bone.Body)
 		{
@@ -329,16 +344,17 @@ void FPBIKSolver::UpdateBonesFromBodies()
 			}
 
 			// optionally pin rotation to that of effector
-			const FQuat RotationWithoutEffector = Bone.Parent ? Bone.Parent->Rotation * Bone.LocalRotationFromInput : Bone.Rotation; 
-			const float RotAmount = FMath::Clamp(Effector.Settings.PinRotation, 0.0f, 1.0f);
-			Bone.Rotation = FQuat::Slerp(RotationWithoutEffector, Effector.RotationFinal, RotAmount);
-		}
-
-		// trivial leaf bones (children of simulated bones)
-		if (!Bone.bIsSolved && Bone.Parent)
-		{
-			Bone.Position = Bone.Parent->Position + Bone.Parent->Rotation * Bone.LocalPositionFromInput;
-			Bone.Rotation = Bone.Parent->Rotation * Bone.LocalRotationFromInput;
+			if (FMath::IsNearlyEqual(Effector.Settings.PinRotation, 1.0f))
+			{
+				Bone.Rotation = Effector.RotationFinal;
+			}
+			else
+			{
+				// blend between bone rotation and effector rotation
+				const FQuat RotationWithoutEffector = Bone.Parent ? Bone.Parent->Rotation * Bone.LocalRotationFromInput : Bone.Rotation; 
+				const float RotAmount = FMath::Clamp(Effector.Settings.PinRotation, 0.0f, 1.0f);
+				Bone.Rotation = FQuat::FastLerp(RotationWithoutEffector, Effector.RotationFinal, RotAmount).GetNormalized();
+			}	
 		}
 	}
 }
@@ -388,14 +404,14 @@ void FPBIKSolver::ApplyRootPrePull(const EPBIKRootBehavior RootBehavior, const F
 	// move and rotate all bodies rigidly with the root
 	const FRigidBody* RootBody = SolverRoot->Body;
 	const FVector RootPosition = RootBody->Position;
-	const FQuat RootRotation = RootBody->Rotation;
+	const FQuat RootRotationInv = RootBody->Rotation.Inverse();
 	const FVector NewRootPosition = RootPosition + PositionOffset;
 	const FQuat NewRootRotation = FinalRotationOffset * RootBody->Rotation;
 	for (FRigidBody& Body : Bodies)
 	{
-		FVector LocalPos = RootRotation.Inverse() * (Body.Position - RootPosition);
-		FQuat LocalRot = RootRotation.Inverse() * Body.Rotation;
-		Body.Position = NewRootPosition + (NewRootRotation * LocalPos);
+		FVector BodyLocalPos = RootRotationInv * (Body.Position - RootPosition);
+		FQuat LocalRot = RootRotationInv * Body.Rotation;
+		Body.Position = NewRootPosition + (NewRootRotation * BodyLocalPos);
 		Body.Rotation = NewRootRotation * LocalRot;
 	}
 }
@@ -428,34 +444,31 @@ void FPBIKSolver::ApplyPullChainAlpha(const float GlobalPullChainAlpha)
 			continue;
 		}
 
-		// get original chain vector
-		const FVector ChainStartOrig = Effector.ChainRootBody->InputPosition;
-		const FVector ChainEndOrig = Effector.PositionOrig;
-		FVector ChainVecOrig;
-		float ChainLenOrig;
-		(ChainEndOrig - ChainStartOrig).ToDirectionAndLength(ChainVecOrig, ChainLenOrig);
+		// get the current start point on the chain
+		FVector ChainStartPosition = Effector.ChainRootBody->Position;
 
-		// get new chain vector
-		FVector ChainStartNew;
+		// if the chain root has an effector we have to add it's offset here
 		if (Effector.ChainRootBody->Effector)
 		{
-			// add effector offset to chain root
 			const FEffector* RootEffector = Effector.ChainRootBody->Effector;
 			FVector RootEffectorDelta = RootEffector->PositionFinal - RootEffector->PositionOrig;
-			ChainStartNew = Effector.ChainRootBody->Position + RootEffectorDelta;
-		}else
-		{
-			ChainStartNew = Effector.ChainRootBody->Position;
+			ChainStartPosition += RootEffectorDelta;
 		}
-		const FVector ChainEndNew = Effector.PositionFinal;
+
+		// get current chain vector
+		FVector ChainVecCurrent;
+		float ChainLenCurrent;
+		(Effector.Pin->GetCurrentPinPoint() - ChainStartPosition).ToDirectionAndLength(ChainVecCurrent, ChainLenCurrent);
+
+		// get new chain vector
 		FVector ChainVecNew;
 		float ChainLenNew;
-		(ChainEndNew - ChainStartNew).ToDirectionAndLength(ChainVecNew, ChainLenNew);
+		(Effector.PositionFinal - ChainStartPosition).ToDirectionAndLength(ChainVecNew, ChainLenNew);
 
-		FQuat ChainDeltaRotation = FQuat::FindBetweenNormals(ChainVecOrig, ChainVecNew);
+		FQuat ChainDeltaRotation = FQuat::FindBetweenNormals(ChainVecCurrent, ChainVecNew);
 		const float ChainDeltaAlpha = Effector.Settings.PullChainAlpha * Effector.Settings.StrengthAlpha * GlobalPullChainAlpha;
 		ChainDeltaRotation = FQuat::FastLerp(FQuat::Identity, ChainDeltaRotation, ChainDeltaAlpha).GetNormalized();
-		FVector ChainDeltaPosition = ChainVecNew * (ChainLenNew - ChainLenOrig) * ChainDeltaAlpha;
+		FVector ChainDeltaPosition = ChainVecNew * (ChainLenNew - ChainLenCurrent) * ChainDeltaAlpha;
 
 		const FBone* Bone = Effector.Bone->Parent;
 		int32 ChainIndex = 0;
@@ -466,8 +479,8 @@ void FPBIKSolver::ApplyPullChainAlpha(const float GlobalPullChainAlpha)
 			if (!Bone->Body->IsPositionLocked())
 			{
 				// move first to the rotated chain position
-				const FVector BodyRelativeToChain = Bone->Body->Position - ChainStartNew;
-				const FVector RotatedBodyPos = ChainStartNew + ChainDeltaRotation.RotateVector(BodyRelativeToChain);
+				const FVector BodyRelativeToChain = Bone->Body->Position - ChainStartPosition;
+				const FVector RotatedBodyPos = ChainStartPosition + ChainDeltaRotation.RotateVector(BodyRelativeToChain);
 				Bone->Body->Position = RotatedBodyPos;
 
 				// translate along chain proportional to chain stretch amount
@@ -659,6 +672,27 @@ bool FPBIKSolver::InitBones()
 		Bone.Parent = &Bones[Bone.ParentIndex];
 	}
 
+	// initialize bIsChildOfSolverRoot
+	for (FBone& Bone : Bones)
+	{
+		if (Bone.bIsSolverRoot)
+		{
+			Bone.bIsPartOfSolvedBranch = true;
+			continue;
+		}
+		
+		FBone* CurrentParent = Bone.Parent;
+		while (CurrentParent)
+		{
+			if (CurrentParent->bIsSolverRoot)
+			{
+				Bone.bIsPartOfSolvedBranch = true;
+				break;
+			}
+			CurrentParent = CurrentParent->Parent;
+		}
+	}
+
 	// walk upwards from each effector to the root to initialize "Bone.IsSolved"
 	for (const FEffector& Effector : Effectors)
 	{
@@ -790,7 +824,7 @@ bool FPBIKSolver::InitConstraints()
 
 		FRigidBody* Body = BodyBone->Body;
 		TSharedPtr<FPinConstraint> Constraint = MakeShared<FPinConstraint>(Body, Effector.Bone);
-		Effector.Pin = Constraint;
+		Effector.Pin = Constraint.Get();
 		Body->Effector = &Effector;
 		Body->Pin = Constraint.Get();
 		Constraints.Add(Constraint);
diff --git a/Engine/Plugins/Experimental/FullBodyIK/Source/PBIK/Public/Core/PBIKBody.h b/Engine/Plugins/Experimental/FullBodyIK/Source/PBIK/Public/Core/PBIKBody.h
--- a/Engine/Plugins/Experimental/FullBodyIK/Source/PBIK/Public/Core/PBIKBody.h
+++ b/Engine/Plugins/Experimental/FullBodyIK/Source/PBIK/Public/Core/PBIKBody.h
@@ -32,6 +32,7 @@ struct FBone
 	bool bIsSolverRoot = false;
 	bool bIsSolved = false;
 	bool bIsSubRoot = false;
+	bool bIsPartOfSolvedBranch = false;
 	float Length = 0.f;
 	// initialized
 
@@ -87,7 +88,6 @@ struct FRigidBody
 	FVector InputPosition;
 	FQuat InitialRotation;
 	FVector BoneLocalPosition;
-	TArray<FVector> ChildLocalPositions;
 
 	float InvMass = 0.f;
 	float Mass = 0.f;
diff --git a/Engine/Plugins/Experimental/FullBodyIK/Source/PBIK/Public/Core/PBIKConstraint.h b/Engine/Plugins/Experimental/FullBodyIK/Source/PBIK/Public/Core/PBIKConstraint.h
--- a/Engine/Plugins/Experimental/FullBodyIK/Source/PBIK/Public/Core/PBIKConstraint.h
+++ b/Engine/Plugins/Experimental/FullBodyIK/Source/PBIK/Public/Core/PBIKConstraint.h
@@ -119,6 +119,8 @@ public:
 
 	void SetGoal(const FVector& InGoalPosition, const FQuat& InGoalRotation, const float InAlpha);
 
+	FVector GetCurrentPinPoint() const;
+
 private:
 
 	FVector GetPositionCorrection(FVector& OutBodyToPinPoint) const;
diff --git a/Engine/Plugins/Experimental/FullBodyIK/Source/PBIK/Public/Core/PBIKSolver.h b/Engine/Plugins/Experimental/FullBodyIK/Source/PBIK/Public/Core/PBIKSolver.h
--- a/Engine/Plugins/Experimental/FullBodyIK/Source/PBIK/Public/Core/PBIKSolver.h
+++ b/Engine/Plugins/Experimental/FullBodyIK/Source/PBIK/Public/Core/PBIKSolver.h
@@ -74,7 +74,7 @@ struct FEffector
 	FEffectorSettings Settings;
 
 	FBone* Bone;
-	TWeakPtr<FPinConstraint> Pin;
+	FPinConstraint* Pin;
 	FRigidBody* ChainRootBody = nullptr;
 	int32 ChainDepthInitializedWith = -1;
 	float DistToChainRootInInputPose;
diff --git a/Engine/Plugins/Experimental/FullBodyIK/Source/PBIK/Public/RigUnit_PBIK.h b/Engine/Plugins/Experimental/FullBodyIK/Source/PBIK/Public/RigUnit_PBIK.h
--- a/Engine/Plugins/Experimental/FullBodyIK/Source/PBIK/Public/RigUnit_PBIK.h
+++ b/Engine/Plugins/Experimental/FullBodyIK/Source/PBIK/Public/RigUnit_PBIK.h
@@ -41,14 +41,6 @@ struct FPBIKDebug
 
 		const FLinearColor Bright = FLinearColor(0.f, 1.f, 1.f, 1.f);
 		DrawInterface->DrawBox(FTransform::Identity, FTransform(FQuat::Identity, FVector(0, 0, 0), FVector(1.f, 1.f, 1.f) * DrawScale * 0.1f), Bright);
-
-		TArray<FDebugLine> BodyLines;
-		Solver->GetDebugDraw()->GetDebugLinesForBodies(BodyLines);
-		const FLinearColor BodyColor = FLinearColor(0.1f, 0.1f, 1.f, 1.f);
-		for (FDebugLine Line : BodyLines)
-		{
-			DrawInterface->DrawLine(FTransform::Identity, Line.A, Line.B, BodyColor);
-		}
 	}
 };
 
