Fullbody IK: Fix non-normalized rotation coming from solver

#rnx

diff --git a/Engine/Plugins/Experimental/FullBodyIK/Source/PBIK/Private/RigUnit_PBIK.cpp b/Engine/Plugins/Experimental/FullBodyIK/Source/PBIK/Private/RigUnit_PBIK.cpp
--- a/Engine/Plugins/Experimental/FullBodyIK/Source/PBIK/Private/RigUnit_PBIK.cpp
+++ b/Engine/Plugins/Experimental/FullBodyIK/Source/PBIK/Private/RigUnit_PBIK.cpp
@@ -199,6 +199,7 @@ FRigUnit_PBIK_Execute()
 	{
 		FTransform NewTransform;
 		WorkData.Solver.GetBoneGlobalTransform(BoneIndex, NewTransform);
+		NewTransform.NormalizeRotation();
 		Hierarchy->SetGlobalTransform(WorkData.SolverBoneToElementIndex[BoneIndex], NewTransform, false, bPropagateTransform);
 	}
 
