[FBIK] Fixed issue with non-normalized rotations exploding the solve.
Fixed issue in Pull Chain Alpha not rotating the first bone.
Added Over Relaxation to constraint rotation

diff --git a/Engine/Plugins/Experimental/FullBodyIK/Source/PBIK/Private/Core/PBIKBody.cpp b/Engine/Plugins/Experimental/FullBodyIK/Source/PBIK/Private/Core/PBIKBody.cpp
--- a/Engine/Plugins/Experimental/FullBodyIK/Source/PBIK/Private/Core/PBIKBody.cpp
+++ b/Engine/Plugins/Experimental/FullBodyIK/Source/PBIK/Private/Core/PBIKBody.cpp
@@ -128,7 +128,8 @@ void FRigidBody::ApplyPushToRotateBody(const FVector& Push, const FVector& Offse
 	}
 	
 	// equation 8 in "Detailed Rigid Body Simulation with XPBD"
-	const FVector Omega = InvMass * (1.0f - J.RotationStiffness) * FVector::CrossProduct(Offset, Push);
+	const float Lambda = InvMass * (1.0f - J.RotationStiffness) * SolverSettings->OverRelaxation;
+	const FVector Omega = Lambda * FVector::CrossProduct(Offset, Push);
 	const FQuat DeltaQ(Omega.X, Omega.Y, Omega.Z, 0.0f);
 	ApplyRotationDelta(DeltaQ);
 }
diff --git a/Engine/Plugins/Experimental/FullBodyIK/Source/PBIK/Private/Core/PBIKConstraint.cpp b/Engine/Plugins/Experimental/FullBodyIK/Source/PBIK/Private/Core/PBIKConstraint.cpp
--- a/Engine/Plugins/Experimental/FullBodyIK/Source/PBIK/Private/Core/PBIKConstraint.cpp
+++ b/Engine/Plugins/Experimental/FullBodyIK/Source/PBIK/Private/Core/PBIKConstraint.cpp
@@ -61,6 +61,10 @@ void FJointConstraint::Solve(const FPBIKSolverSettings& Settings)
 	// apply position correction
 	A->ApplyPositionDelta(CorrectA);
 	B->ApplyPositionDelta(CorrectB);
+
+	// must normalize due to infinitesimal quaternion addition
+	A->Rotation.Normalize();
+	B->Rotation.Normalize();
 }
 
 void FJointConstraint::RemoveStretch(const float Percent)
diff --git a/Engine/Plugins/Experimental/FullBodyIK/Source/PBIK/Private/Core/PBIKSolver.cpp b/Engine/Plugins/Experimental/FullBodyIK/Source/PBIK/Private/Core/PBIKSolver.cpp
--- a/Engine/Plugins/Experimental/FullBodyIK/Source/PBIK/Private/Core/PBIKSolver.cpp
+++ b/Engine/Plugins/Experimental/FullBodyIK/Source/PBIK/Private/Core/PBIKSolver.cpp
@@ -445,12 +445,13 @@ void FPBIKSolver::ApplyPullChainAlpha(const float GlobalPullChainAlpha)
 		}
 
 		// get the current start point on the chain
-		FVector ChainStartPosition = Effector.ChainRootBody->Position;
+		FRigidBody* RootBody = Effector.ChainRootBody;
+		FVector ChainStartPosition = RootBody->Position + RootBody->Rotation * RootBody->BoneLocalPosition;
 
 		// if the chain root has an effector we have to add it's offset here
-		if (Effector.ChainRootBody->Effector)
+		if (RootBody->Effector)
 		{
-			const FEffector* RootEffector = Effector.ChainRootBody->Effector;
+			const FEffector* RootEffector = RootBody->Effector;
 			FVector RootEffectorDelta = RootEffector->PositionFinal - RootEffector->PositionOrig;
 			ChainStartPosition += RootEffectorDelta;
 		}
@@ -473,7 +474,7 @@ void FPBIKSolver::ApplyPullChainAlpha(const float GlobalPullChainAlpha)
 		const FBone* Bone = Effector.Bone->Parent;
 		int32 ChainIndex = 0;
 		const float InvChainLength = 1.0f / Effector.DistToChainRootAlongBones;
-		while (Bone && Bone->Body != Effector.ChainRootBody)
+		while (Bone)
 		{
 			// move body along with chain
 			if (!Bone->Body->IsPositionLocked())
@@ -493,6 +494,11 @@ void FPBIKSolver::ApplyPullChainAlpha(const float GlobalPullChainAlpha)
 			{
 				Bone->Body->Rotation = ChainDeltaRotation * Bone->Body->Rotation;
 			}
+
+			if (Bone->Body == RootBody)
+			{
+				break;
+			}
 			
 			Bone = Bone->Parent;
 			++ChainIndex;
@@ -551,33 +557,46 @@ void FPBIKSolver::SolveConstraints(const FPBIKSolverSettings& InSolverSettings)
 		}
 	};
 
-	// run pre-pass where we lock everything not in a sub-chain
-	if (bHasSubChains && InSolverSettings.SubIterations > 0)
+	auto RunSubIterations = [this, &RunConstraintPass, &InSolverSettings]()
 	{
+		// run pre-pass where we lock everything not in a sub-chain
+		if (!bHasSubChains || InSolverSettings.SubIterations <= 0)
+		{
+			return;
+		}
 		// first lock all the bodies not in a sub chain
 		for (FRigidBody& Body : Bodies)
 		{
 			Body.bIsLockedBySubSolve = !Body.bIsPartOfSubChain;
 		}
-		
+
 		// solve all constraints (locked constraints are skipped in this pass)
 		for (int32  I = 0; I < InSolverSettings.SubIterations; ++I)
 		{
 			RunConstraintPass(I, InSolverSettings.SubIterations);
 		}
-		
+
 		// unlock all bodies so we can proceed with normal constraint solve
 		for (FRigidBody& Body : Bodies)
 		{
 			Body.bIsLockedBySubSolve = false;
 		}
-	}
+	};
 
-	// run main constraint pass
-	for (int32  I = 0; I < InSolverSettings.Iterations; ++I)
+	auto RunMainIterations = [this, &RunConstraintPass, &InSolverSettings]()
 	{
-		RunConstraintPass(I, InSolverSettings.Iterations);
-	}
+		// run main constraint pass
+		for (int32  I = 0; I < InSolverSettings.Iterations; ++I)
+		{
+			RunConstraintPass(I, InSolverSettings.Iterations);
+		}
+	};
+
+	// run the sub iterations with non-chain bodies locked
+	RunSubIterations();
+	
+	// run the main iterations with all bodies un-locked
+	RunMainIterations();
 }
 
 bool FPBIKSolver::Initialize()
diff --git a/Engine/Plugins/Experimental/FullBodyIK/Source/PBIK/Private/RigUnit_PBIK.cpp b/Engine/Plugins/Experimental/FullBodyIK/Source/PBIK/Private/RigUnit_PBIK.cpp
--- a/Engine/Plugins/Experimental/FullBodyIK/Source/PBIK/Private/RigUnit_PBIK.cpp
+++ b/Engine/Plugins/Experimental/FullBodyIK/Source/PBIK/Private/RigUnit_PBIK.cpp
@@ -194,12 +194,11 @@ FRigUnit_PBIK_Execute()
 	WorkData.Solver.Solve(Settings);
 
 	// copy transforms back
-	const bool bPropagateTransform = false;
 	for(int32 BoneIndex = 0; BoneIndex < WorkData.Solver.GetNumBones(); BoneIndex++)
 	{
 		FTransform NewTransform;
 		WorkData.Solver.GetBoneGlobalTransform(BoneIndex, NewTransform);
-		NewTransform.NormalizeRotation();
+		constexpr bool bPropagateTransform = false;
 		Hierarchy->SetGlobalTransform(WorkData.SolverBoneToElementIndex[BoneIndex], NewTransform, false, bPropagateTransform);
 	}
 
