[FBIK] Removed broken old debug drawing from FBIK rig unit.

#FailedValidations Preflight

diff --git a/Engine/Plugins/Experimental/FullBodyIK/Source/PBIK/Private/RigUnit_PBIK.cpp b/Engine/Plugins/Experimental/FullBodyIK/Source/PBIK/Private/RigUnit_PBIK.cpp
--- a/Engine/Plugins/Experimental/FullBodyIK/Source/PBIK/Private/RigUnit_PBIK.cpp
+++ b/Engine/Plugins/Experimental/FullBodyIK/Source/PBIK/Private/RigUnit_PBIK.cpp
@@ -201,8 +201,5 @@ FRigUnit_PBIK_Execute()
 		constexpr bool bPropagateTransform = false;
 		Hierarchy->SetGlobalTransform(WorkData.SolverBoneToElementIndex[BoneIndex], NewTransform, false, bPropagateTransform);
 	}
-
-	// do all debug drawing
-	Debug.Draw(ExecuteContext.GetDrawInterface(), &WorkData.Solver);
 }
 
diff --git a/Engine/Plugins/Experimental/FullBodyIK/Source/PBIK/Public/RigUnit_PBIK.h b/Engine/Plugins/Experimental/FullBodyIK/Source/PBIK/Public/RigUnit_PBIK.h
--- a/Engine/Plugins/Experimental/FullBodyIK/Source/PBIK/Public/RigUnit_PBIK.h
+++ b/Engine/Plugins/Experimental/FullBodyIK/Source/PBIK/Public/RigUnit_PBIK.h
@@ -17,33 +17,6 @@
 
 struct FControlRigExecuteContext;
 
-using PBIK::FDebugLine;
-
-USTRUCT(BlueprintType)
-struct FPBIKDebug
-{
-	GENERATED_BODY()
-
-	/** The scale of the debug drawing. */
-	UPROPERTY(EditAnywhere, Category = Debug)
-	float DrawScale = 1.0f;
-
-	/** If true, turns on debug drawing for the node. */
-	UPROPERTY(EditAnywhere, Category = Debug)
-	bool bDrawDebug = false;
-
-	void Draw(FRigVMDrawInterface* DrawInterface, FPBIKSolver* Solver) const
-	{
-		if (!(DrawInterface && Solver && bDrawDebug))
-		{
-			return;
-		}
-
-		const FLinearColor Bright = FLinearColor(0.f, 1.f, 1.f, 1.f);
-		DrawInterface->DrawBox(FTransform::Identity, FTransform(FQuat::Identity, FVector(0, 0, 0), FVector(1.f, 1.f, 1.f) * DrawScale * 0.1f), Bright);
-	}
-};
-
 USTRUCT(BlueprintType)
 struct FPBIKEffector
 {
@@ -168,10 +141,6 @@ struct FRigUnit_PBIK : public FRigUnit_HighlevelBaseMutable
 	UPROPERTY(meta = (Input))
 	FPBIKSolverSettings Settings;
 
-	/** Debug drawing options. */
-	UPROPERTY(meta = (Input))
-	FPBIKDebug Debug;
-
 	/** Runtime-only data */
 	UPROPERTY(transient)
 	FPBIKWorkData WorkData;
