[FBIK] Fixed crash bug from intermediate effectors on fork joints.

diff --git a/Engine/Plugins/Experimental/FullBodyIK/Source/PBIK/Private/Core/PBIKSolver.cpp b/Engine/Plugins/Experimental/FullBodyIK/Source/PBIK/Private/Core/PBIKSolver.cpp
--- a/Engine/Plugins/Experimental/FullBodyIK/Source/PBIK/Private/Core/PBIKSolver.cpp
+++ b/Engine/Plugins/Experimental/FullBodyIK/Source/PBIK/Private/Core/PBIKSolver.cpp
@@ -90,12 +90,17 @@ namespace PBIK
 		if (ChainRootBody)
 		{
 			DistancesFromEffector.Reset();
-			DistancesFromEffector.Add(0.0f);
-			DistToChainRootAlongBones = Bone->Length;
-			const FBone* Parent = Bone->Parent;
-			while (Parent && Parent->Body != ChainRootBody)
+			DistToChainRootAlongBones = 0.0;
+			const FBone* Parent = Bone;
+			while (true)
 			{
 				DistancesFromEffector.Add(DistToChainRootAlongBones);
+
+				if (!Parent || Parent->Body == ChainRootBody)
+				{
+					break; // stop before advancing past the chain root
+				}
+				
 				DistToChainRootAlongBones += Parent->Length;
 				Parent = Parent->Parent;
 			}
@@ -744,7 +749,7 @@ bool FPBIKSolver::InitBones()
 	// initialize IsSubRoot flag
 	for (FBone& Bone : Bones)
 	{
-		Bone.bIsSubRoot = Bone.Children.Num() > 1 || Bone.bIsSolverRoot;
+		Bone.bIsSubRoot = Bone.Children.Num() > 1 || Bone.bIsSolverRoot || Bone.EffectorIndex != INDEX_NONE;
 	}
 
 	// store initial local rotation (preferred angles are relative to this)
