[FBIK] Added back previously removed debug properties.

These were removed for lack of usefulness, but caused issues with existing assets. They can be improved in the future.

#Synced-CL 48536607

diff --git a/Engine/Plugins/Experimental/FullBodyIK/Source/PBIK/Public/RigUnit_PBIK.h b/Engine/Plugins/Experimental/FullBodyIK/Source/PBIK/Public/RigUnit_PBIK.h
--- a/Engine/Plugins/Experimental/FullBodyIK/Source/PBIK/Public/RigUnit_PBIK.h
+++ b/Engine/Plugins/Experimental/FullBodyIK/Source/PBIK/Public/RigUnit_PBIK.h
@@ -17,15 +17,31 @@
 
 struct FControlRigExecuteContext;
 
-USTRUCT()
-struct  UE_DEPRECATED(5.5, "FPBIKDebug is deprecated and usued. Only here for loading old assets.") FPBIKDebug
+using PBIK::FDebugLine;
+
+USTRUCT(BlueprintType)
+struct FPBIKDebug
 {
 	GENERATED_BODY()
-	
-	UPROPERTY(meta=(DeprecatedProperty, DeprecationMessage="Rig Unit no longer does debug drawing."))
+
+	/** The scale of the debug drawing. */
+	UPROPERTY(EditAnywhere, Category = Debug)
 	float DrawScale = 1.0f;
-	UPROPERTY(meta=(DeprecatedProperty, DeprecationMessage="Rig Unit no longer does debug drawing."))
+
+	/** If true, turns on debug drawing for the node. */
+	UPROPERTY(EditAnywhere, Category = Debug)
 	bool bDrawDebug = false;
+
+	void Draw(FRigVMDrawInterface* DrawInterface, FPBIKSolver* Solver) const
+	{
+		if (!(DrawInterface && Solver && bDrawDebug))
+		{
+			return;
+		}
+
+		const FLinearColor Bright = FLinearColor(0.f, 1.f, 1.f, 1.f);
+		DrawInterface->DrawBox(FTransform::Identity, FTransform(FQuat::Identity, FVector(0, 0, 0), FVector(1.f, 1.f, 1.f) * DrawScale * 0.1f), Bright);
+	}
 };
 
 USTRUCT(BlueprintType)
@@ -116,7 +132,7 @@ struct FPBIKWorkData
 /*
  * Based on a Position Based solver at core, this node can solve multi chains within a root using multi effectors
  */
-USTRUCT(meta=(DisplayName="Full Body IK", Category="Hierarchy", Keywords="FBIK, Position Based, PBIK, IK, Full Body, Multi, Effector, N-Chain, FB, HIK, HumanIK", NodeColor="0 1 1"))
+USTRUCT(meta=(DisplayName="Full Body IK", NodeColor="0 1 1"))
 struct FRigUnit_PBIK : public FRigUnit_HighlevelBaseMutable
 {
 	GENERATED_BODY()
@@ -152,6 +168,10 @@ struct FRigUnit_PBIK : public FRigUnit_HighlevelBaseMutable
 	UPROPERTY(meta = (Input))
 	FPBIKSolverSettings Settings;
 
+	/** Debug drawing options. */
+	UPROPERTY(meta = (Input))
+	FPBIKDebug Debug;
+
 	/** Runtime-only data */
 	UPROPERTY(transient)
 	FPBIKWorkData WorkData;
